[PS2] Fix: Renamed media texture factory to not contain EpicRtc reference

diff --git a/Engine/Plugins/Media/PixelStreaming2/Source/PixelStreaming2BlueprintEditor/Private/PixelStreaming2MediaTextureFactory.cpp b/Engine/Plugins/Media/PixelStreaming2/Source/PixelStreaming2BlueprintEditor/Private/PixelStreaming2MediaTextureFactory.cpp
--- a/Engine/Plugins/Media/PixelStreaming2/Source/PixelStreaming2BlueprintEditor/Private/PixelStreaming2MediaTextureFactory.cpp
+++ b/Engine/Plugins/Media/PixelStreaming2/Source/PixelStreaming2BlueprintEditor/Private/PixelStreaming2MediaTextureFactory.cpp
@@ -18,7 +18,7 @@ UPixelStreaming2MediaTextureFactory::UPixelStreaming2MediaTextureFactory(const F
 
 FText UPixelStreaming2MediaTextureFactory::GetDisplayName() const
 {
-	return LOCTEXT("MediaTextureFactoryDisplayName", "Pixel Streaming EpicRtc Media Texture");
+	return LOCTEXT("MediaTextureFactoryDisplayName", "Pixel Streaming 2 Media Texture");
 }
 
 uint32 UPixelStreaming2MediaTextureFactory::GetMenuCategories() const
