[PS2] Fix: Clashing CVar names

diff --git a/Engine/Plugins/Media/PixelStreaming2/Source/PixelStreaming2/Private/Stats.cpp b/Engine/Plugins/Media/PixelStreaming2/Source/PixelStreaming2/Private/Stats.cpp
--- a/Engine/Plugins/Media/PixelStreaming2/Source/PixelStreaming2/Private/Stats.cpp
+++ b/Engine/Plugins/Media/PixelStreaming2/Source/PixelStreaming2/Private/Stats.cpp
@@ -216,13 +216,13 @@ namespace UE::PixelStreaming2
 
 		AutoCompleteList.AddDefaulted();
 		FAutoCompleteCommand& AutoCompleteCommand = AutoCompleteList.Last();
-		AutoCompleteCommand.Command = TEXT("Stat PixelStreaming");
+		AutoCompleteCommand.Command = TEXT("Stat PixelStreaming2");
 		AutoCompleteCommand.Desc = TEXT("Displays stats about Pixel Streaming on screen.");
 		AutoCompleteCommand.Color = ConsoleSettings->AutoCompleteCommandColor;
 
 		AutoCompleteList.AddDefaulted();
 		FAutoCompleteCommand& AutoCompleteGraphCommand = AutoCompleteList.Last();
-		AutoCompleteGraphCommand.Command = TEXT("Stat PixelStreamingGraphs");
+		AutoCompleteGraphCommand.Command = TEXT("Stat PixelStreaming2Graphs");
 		AutoCompleteGraphCommand.Desc = TEXT("Displays graphs about Pixel Streaming on screen.");
 		AutoCompleteGraphCommand.Color = ConsoleSettings->AutoCompleteCommandColor;
 	}
@@ -451,7 +451,7 @@ namespace UE::PixelStreaming2
 
 		// We register this console command so we get autocomplete on `Stat PixelStreaming2`
 		IConsoleManager::Get().RegisterConsoleCommand(
-			TEXT("Stat PixelStreaming"),
+			TEXT("Stat PixelStreaming2"),
 			TEXT("Stats for the Pixel Streaming plugin and its peers."),
 			FConsoleCommandDelegate::CreateRaw(this, &FStats::ExecStatPS),
 			ECVF_Default);
@@ -464,7 +464,7 @@ namespace UE::PixelStreaming2
 
 		// We register this console command so we get autocomplete on `Stat PixelStreaming2Graphs`
 		IConsoleManager::Get().RegisterConsoleCommand(
-			TEXT("Stat PixelStreamingGraphs"),
+			TEXT("Stat PixelStreaming2Graphs"),
 			TEXT("Draws stats graphs for the Pixel Streaming plugin."),
 			FConsoleCommandDelegate::CreateRaw(this, &FStats::ExecStatPSGraphs),
 			ECVF_Default);
diff --git a/Engine/Plugins/Media/PixelStreaming2/Source/PixelStreaming2/Private/Stats.h b/Engine/Plugins/Media/PixelStreaming2/Source/PixelStreaming2/Private/Stats.h
--- a/Engine/Plugins/Media/PixelStreaming2/Source/PixelStreaming2/Private/Stats.h
+++ b/Engine/Plugins/Media/PixelStreaming2/Source/PixelStreaming2/Private/Stats.h
@@ -133,8 +133,8 @@ namespace UE::PixelStreaming2
 
 	// Stats about Pixel Streaming that can displayed either in the in-application HUD, in the log, or simply reported to some subscriber.
 	// Stats can be enabled to draw on screen with:
-	// `stat pixelstreaming`
-	// `stat pixelstreaminggraphs`
+	// `stat pixelstreaming2`
+	// `stat pixelstreaming2graphs`
 	class FStats : public FTickableGameObject, public IPixelStreaming2Stats
 	{
 	public:
diff --git a/Engine/Plugins/Media/PixelStreaming2/Source/PixelStreaming2Core/Internal/PixelStreaming2PluginSettings.h b/Engine/Plugins/Media/PixelStreaming2/Source/PixelStreaming2Core/Internal/PixelStreaming2PluginSettings.h
--- a/Engine/Plugins/Media/PixelStreaming2/Source/PixelStreaming2Core/Internal/PixelStreaming2PluginSettings.h
+++ b/Engine/Plugins/Media/PixelStreaming2/Source/PixelStreaming2Core/Internal/PixelStreaming2PluginSettings.h
@@ -75,7 +75,7 @@ public:
 	// clang-format off
 	static TAutoConsoleVariable<bool> CVarOnScreenStats;
 	UPROPERTY(config, EditAnywhere, Category = "PixelStreaming", meta = (
-		ConsoleVariable = "PixelStreaming.HUDStats",
+		ConsoleVariable = "PixelStreaming2.HUDStats",
 		DisplayName = "Show Pixel Streaming Stats On Screen",
 		ToolTip = "Whether to show PixelStreaming stats on the in-game HUD (default: false)."
 		))
@@ -83,7 +83,7 @@ public:
 
 	static TAutoConsoleVariable<bool> CVarLogStats;
 	UPROPERTY(config, EditAnywhere, Category = "PixelStreaming", meta = (
-		ConsoleVariable = "PixelStreaming.LogStats",
+		ConsoleVariable = "PixelStreaming2.LogStats",
 		DisplayName = "Log Pixel Streaming Stats",
 		ToolTip = "Whether to show PixelStreaming stats in the log (default: false)."
 		))
@@ -91,7 +91,7 @@ public:
 
 	static TAutoConsoleVariable<bool> CVarDisableLatencyTester;
 	UPROPERTY(config, EditAnywhere, Category = "PixelStreaming", meta = (
-		ConsoleVariable = "PixelStreaming.DisableLatencyTester",
+		ConsoleVariable = "PixelStreaming2.DisableLatencyTester",
 		DisplayName = "Disable Latency Tester",
 		ToolTip = "If true disables latency tester being triggerable."
 		))
@@ -99,7 +99,7 @@ public:
 
 	static TAutoConsoleVariable<FString> CVarInputController;
 	UPROPERTY(config, EditAnywhere, Category = "PixelStreaming", meta = (
-		MappedConsoleVariable = "PixelStreaming.InputController",
+		MappedConsoleVariable = "PixelStreaming2.InputController",
 		DisplayName = "Input Controller Mode",
 		ToolTip = "If true disables latency tester being triggerable."
 		))
@@ -107,7 +107,7 @@ public:
 
 	static TAutoConsoleVariable<bool> CVarDecoupleFramerate;
 	UPROPERTY(config, EditAnywhere, Category = "PixelStreaming", meta = (
-		ConsoleVariable = "PixelStreaming.DecoupleFrameRate",
+		ConsoleVariable = "PixelStreaming2.DecoupleFrameRate",
 		DisplayName = "Decouple Frame Rate",
 		ToolTip = "Whether we should only stream as fast as we render or at some fixed interval. Coupled means only stream what we render."
 		))
@@ -115,7 +115,7 @@ public:
 
 	static TAutoConsoleVariable<float> CVarDecoupleWaitFactor;
 	UPROPERTY(config, EditAnywhere, Category = "PixelStreaming", meta = (
-		ConsoleVariable = "PixelStreaming.DecoupleWaitFactor",
+		ConsoleVariable = "PixelStreaming2.DecoupleWaitFactor",
 		DisplayName = "Decouple Wait Factor",
 		ToolTip = "Frame rate factor to wait for a captured frame when streaming in decoupled mode. Higher factor waits longer but may also result in higher latency."
 		))
@@ -123,7 +123,7 @@ public:
 
 	static TAutoConsoleVariable<float> CVarSignalingReconnectInterval;
 	UPROPERTY(config, EditAnywhere, Category = "PixelStreaming", meta = (
-		ConsoleVariable = "PixelStreaming.SignalingReconnectInterval",
+		ConsoleVariable = "PixelStreaming2.SignalingReconnectInterval",
 		DisplayName = "Signaling Reconnect Interval",
 		ToolTip = "Frame rate factor to wait for a captured frame when streaming in decoupled mode. Higher factor waits longer but may also result in higher latency."
 		))
@@ -131,7 +131,7 @@ public:
 
 	static TAutoConsoleVariable<float> CVarSignalingKeepAliveInterval;
 	UPROPERTY(config, EditAnywhere, Category = "PixelStreaming", meta = (
-		ConsoleVariable = "PixelStreaming.SignalingKeepAliveInterval",
+		ConsoleVariable = "PixelStreaming2.SignalingKeepAliveInterval",
 		DisplayName = "Signaling Reconnect Interval",
 		ToolTip = "Frame rate factor to wait for a captured frame when streaming in decoupled mode. Higher factor waits longer but may also result in higher latency."
 		))
@@ -139,7 +139,7 @@ public:
 
 	static TAutoConsoleVariable<bool> CVarUseMediaCapture;
 	UPROPERTY(config, EditAnywhere, Category = "PixelStreaming", meta = (
-		ConsoleVariable = "PixelStreaming.UseMediaCapture",
+		ConsoleVariable = "PixelStreaming2.UseMediaCapture",
 		DisplayName = "Use Media Capture",
 		ToolTip = "Use Media Capture from MediaIOFramework to capture frames rather than Pixel Streamings internal backbuffer sources."
 		))
@@ -147,7 +147,7 @@ public:
 
 	static TAutoConsoleVariable<FString> CVarDefaultStreamerID;
 	UPROPERTY(config, EditAnywhere, Category = "PixelStreaming", meta = (
-		ConsoleVariable = "PixelStreaming.ID",
+		ConsoleVariable = "PixelStreaming2.ID",
 		DisplayName = "Default Streamer ID",
 		ToolTip = "Default Streamer ID to be used when not specified elsewhere."
 		))
@@ -155,7 +155,7 @@ public:
 
 	static TAutoConsoleVariable<FString> CVarSignallingURL;
 	UPROPERTY(config, EditAnywhere, Category = "PixelStreaming", meta = (
-		ConsoleVariable = "PixelStreaming.SignallingURL",
+		ConsoleVariable = "PixelStreaming2.SignallingURL",
 		DisplayName = "Default Signalling URL",
 		ToolTip = "Default URL to connect to for signalling."
 		))
@@ -184,7 +184,7 @@ public:
 	// Begin Encoder Settings
 	static TAutoConsoleVariable<int32> CVarEncoderTargetBitrate;
 	UPROPERTY(config, EditAnywhere, Category = "EncoderSettings", meta = (
-		ConsoleVariable = "PixelStreaming.Encoder.TargetBitrate",
+		ConsoleVariable = "PixelStreaming2.Encoder.TargetBitrate",
 		DisplayName = "Target Bitrate",
 		ToolTip = "Target bitrate (bps). Ignore the bitrate WebRTC wants (not recommended). Set to -1 to disable. Default -1.",
 		ClampMin = -1))
@@ -192,7 +192,7 @@ public:
 
 	static TAutoConsoleVariable<int32>	 CVarEncoderMinQuality;
 	UPROPERTY(config, EditAnywhere, Category = "EncoderSettings", meta = (
-		ConsoleVariable = "PixelStreaming.Encoder.MinQuality",
+		ConsoleVariable = "PixelStreaming2.Encoder.MinQuality",
 		DisplayName = "Encoder Minimum Quality",
 		ToolTip = "0-100, Higher values result in a better minimum quality but higher average bitrates. Default 0 - i.e. no limit on a minimum Quality.",
 		ClampMin = 0, ClampMax = 100))
@@ -200,7 +200,7 @@ public:
 
 	static TAutoConsoleVariable<int32>	CVarEncoderMaxQuality;
 	UPROPERTY(config, EditAnywhere, Category = "EncoderSettings", meta = (
-		ConsoleVariable = "PixelStreaming.Encoder.MaxQuality",
+		ConsoleVariable = "PixelStreaming2.Encoder.MaxQuality",
 		DisplayName = "Encoder Maximum Quality",
 		ToolTip = "0-100, Lower values result in lower average bitrates but reduces maximum quality. Default 100 - i.e. no limit on a maximum QP.",
 		ClampMin = 0, ClampMax = 100))
@@ -209,7 +209,7 @@ public:
 	// not directly connected to cvar due to string to enum conversion
 	static TAutoConsoleVariable<FString> CVarEncoderQualityPreset;
 	UPROPERTY(config, EditAnywhere, Category = "EncoderSettings", meta = (
-		MappedConsoleVariable = "PixelStreaming.Encoder.QualityPreset",
+		MappedConsoleVariable = "PixelStreaming2.Encoder.QualityPreset",
 		DisplayName = "Default Encoding Quality Profile",
 		ToolTip = "PixelStreaming Encoder presets that affecting Quality vs Bitrate. Supported modes are `ULTRA_LOW_QUALITY`, `LOW_QUALITY`, `DEFAULT`, `HIGH_QUALITY`, `LOSSLESS`"
 		))
@@ -218,7 +218,7 @@ public:
 	// not directly connected to cvar due to string to enum conversion
 	static TAutoConsoleVariable<FString> CVarEncoderLatencyMode;
 	UPROPERTY(config, EditAnywhere, Category = "EncoderSettings", meta = (
-		MappedConsoleVariable = "PixelStreaming.Encoder.LatencyMode",
+		MappedConsoleVariable = "PixelStreaming2.Encoder.LatencyMode",
 		DisplayName = "Default Encoding Quality Profile",
 		ToolTip = "PixelStreaming Encoder presets that affecting Quality vs Latency. Supported modes are `ULTRA_LOW_LATENCY`, `LOW_LATENCY`, `DEFAULT`"
 		))
@@ -226,7 +226,7 @@ public:
 
 	static TAutoConsoleVariable<int32> CVarEncoderKeyframeInterval;
 	UPROPERTY(config, EditAnywhere, Category = "EncoderSettings", meta = (
-		ConsoleVariable = "PixelStreaming.Encoder.KeyframeInterval",
+		ConsoleVariable = "PixelStreaming2.Encoder.KeyframeInterval",
 		DisplayName = "Default Keyframe Interval",
 		ToolTip = "How many frames before a key frame is sent. Default: -1 which disables the sending of periodic key frames. Note: NVENC reqires a reinitialization when this changes."
 		))
@@ -234,7 +234,7 @@ public:
 
 	static TAutoConsoleVariable<int32> CVarEncoderMaxSessions;
 	UPROPERTY(config, EditAnywhere, Category = "EncoderSettings", meta = (
-		ConsoleVariable = "PixelStreaming.Encoder.MaxSessions",
+		ConsoleVariable = "PixelStreaming2.Encoder.MaxSessions",
 		DisplayName = "Default Max Number of Encoding Session",
 		ToolTip = "-1 implies no limit. Maximum number of concurrent hardware encoder sessions for Pixel Streaming. Note: Geforce gpus only support 8 concurrent sessions and will rollover to software encoding when that number is exceeded."
 		))
@@ -242,7 +242,7 @@ public:
 
 	static TAutoConsoleVariable<bool> CVarEncoderEnableSimulcast;
 	UPROPERTY(config, EditAnywhere, Category = "EncoderSettings", meta = (
-		ConsoleVariable = "PixelStreaming.Encoder.EnableSimulcast",
+		ConsoleVariable = "PixelStreaming2.Encoder.EnableSimulcast",
 		DisplayName = "Enable Simulcast",
 		ToolTip = "Enables simulcast. When enabled, the encoder will encode at full resolution, 1/2 resolution and 1/4 resolution simultaneously. Note: Simulcast is only supported with `H264` and `VP8` and you must use the SFU from the infrastructure to fully utilise this functionality."
 		))
@@ -250,7 +250,7 @@ public:
 
 	static TAutoConsoleVariable<FString> CVarEncoderCodec;
 	UPROPERTY(config, EditAnywhere, Category = "EncoderSettings", meta = (
-		ConsoleVariable = "PixelStreaming.Encoder.Codec",
+		ConsoleVariable = "PixelStreaming2.Encoder.Codec",
 		DisplayName = "Preferred Encoder Codec",
 		ToolTip = "Preferred encoder codec signalled during connection establishment.",
 		GetOptions = "GetVideoCodecOptions"
@@ -259,7 +259,7 @@ public:
 
 	static TAutoConsoleVariable<FString> CVarEncoderScalabilityMode;
 	UPROPERTY(config, EditAnywhere, Category = "EncoderSettings", meta = (
-		ConsoleVariable = "PixelStreaming.Encoder.ScalabilityMode",
+		ConsoleVariable = "PixelStreaming2.Encoder.ScalabilityMode",
 		DisplayName = "Default Scalability Mode",
 		ToolTip = "Indicates number of spatial and temporal layers used, default: L1T1. For a full list of values refer to https://www.w3.org/TR/webrtc-svc/#scalabilitymodes*",
 		GetOptions = "GetScalabilityModeOptions"
@@ -268,7 +268,7 @@ public:
 
 	static TAutoConsoleVariable<FString> CVarEncoderH264Profile;
 	UPROPERTY(config, EditAnywhere, Category = "EncoderSettings|H264", meta = (
-		MappedConsoleVariable = "PixelStreaming.Encoder.H264Profile",
+		MappedConsoleVariable = "PixelStreaming2.Encoder.H264Profile",
 		DisplayName = "Default H264 profile",
 		ToolTip = "PixelStreaming Encoder profile. Supported modes are `AUTO`, `BASELINE`, `MAIN`, `HIGH`, `HIGH444`, `PROGRESSIVE_HIGH` or `CONSTRAINED_HIGH`",
 		ValidEnumValues = "Auto, Baseline, Main, High, High444, ProgressiveHigh, ConstrainedHigh"
@@ -281,7 +281,7 @@ public:
 	// Begin WebRTC CVars
 	static TAutoConsoleVariable<int32> CVarWebRTCFps;
 	UPROPERTY(config, EditAnywhere, Category = "WebRtcSettings", meta = (
-		ConsoleVariable = "PixelStreaming.WebRTC.Fps",
+		ConsoleVariable = "PixelStreaming2.WebRTC.Fps",
 		DisplayName = "Default WebRtc FPS",
 		ToolTip = "Framerate for WebRTC encoding. Default: 60"
 		))
@@ -289,7 +289,7 @@ public:
 
 	static TAutoConsoleVariable<int32> CVarWebRTCStartBitrate;
 	UPROPERTY(config, EditAnywhere, Category = "WebRtcSettings", meta = (
-		ConsoleVariable = "PixelStreaming.WebRTC.StartBitrate",
+		ConsoleVariable = "PixelStreaming2.WebRTC.StartBitrate",
 		DisplayName = "Default WebRtc Start Bitrate",
 		ToolTip = "FStart bitrate (bps) that WebRTC will try begin the stream with. Must be between Min/Max bitrates. Default: 1000000"
 		))
@@ -297,7 +297,7 @@ public:
 
 	static TAutoConsoleVariable<int32> CVarWebRTCMinBitrate;
 	UPROPERTY(config, EditAnywhere, Category = "WebRtcSettings", meta = (
-		ConsoleVariable = "PixelStreaming.WebRTC.MinBitrate",
+		ConsoleVariable = "PixelStreaming2.WebRTC.MinBitrate",
 		DisplayName = "Default WebRtc Minimum Bitrate",
 		ToolTip = "Min bitrate (bps) that WebRTC will not request below. Careful not to set too high otherwise WebRTC will just drop frames. Default: 100000"
 		))
@@ -305,7 +305,7 @@ public:
 
 	static TAutoConsoleVariable<int32> CVarWebRTCMaxBitrate;
 	UPROPERTY(config, EditAnywhere, Category = "WebRtcSettings", meta = (
-		ConsoleVariable = "PixelStreaming.WebRTC.MaxBitrate",
+		ConsoleVariable = "PixelStreaming2.WebRTC.MaxBitrate",
 		DisplayName = "Default WebRtc Maximum Bitrate",
 		ToolTip = "Max bitrate (bps) that WebRTC will not request above. Default: 40000000 aka 40 megabits/per second."
 		))
@@ -313,7 +313,7 @@ public:
 
 	static TAutoConsoleVariable<bool> CVarWebRTCDisableReceiveAudio;
 	UPROPERTY(config, EditAnywhere, Category = "WebRtcSettings", meta = (
-		ConsoleVariable = "PixelStreaming.WebRTC.DisableReceiveAudio",
+		ConsoleVariable = "PixelStreaming2.WebRTC.DisableReceiveAudio",
 		DisplayName = "Disable receiving audio by default",
 		ToolTip = "Disables receiving audio from the browser into UE."
 		))
@@ -321,7 +321,7 @@ public:
 
 	static TAutoConsoleVariable<bool> CVarWebRTCDisableReceiveVideo;
 	UPROPERTY(config, EditAnywhere, Category = "WebRtcSettings", meta = (
-		ConsoleVariable = "PixelStreaming.WebRTC.DisableReceiveVideo",
+		ConsoleVariable = "PixelStreaming2.WebRTC.DisableReceiveVideo",
 		DisplayName = "Disable receiving video by default",
 		ToolTip = "Disables receiving video from the browser into UE."
 		))
@@ -329,7 +329,7 @@ public:
 
 	static TAutoConsoleVariable<bool> CVarWebRTCDisableTransmitAudio;
 	UPROPERTY(config, EditAnywhere, Category = "WebRtcSettings", meta = (
-		ConsoleVariable = "PixelStreaming.WebRTC.DisableTransmitAudio",
+		ConsoleVariable = "PixelStreaming2.WebRTC.DisableTransmitAudio",
 		DisplayName = "Disable transmiting audio by default",
 		ToolTip = "Disables transmiting audio to the browser."
 		))
@@ -337,7 +337,7 @@ public:
 
 	static TAutoConsoleVariable<bool> CVarWebRTCDisableTransmitVideo;
 	UPROPERTY(config, EditAnywhere, Category = "WebRtcSettings", meta = (
-		ConsoleVariable = "PixelStreaming.WebRTC.DisableTransmitVideo",
+		ConsoleVariable = "PixelStreaming2.WebRTC.DisableTransmitVideo",
 		DisplayName = "Disable transmiting video by default",
 		ToolTip = "Disables transmiting video to the browser."
 		))
@@ -345,7 +345,7 @@ public:
 
 	static TAutoConsoleVariable<bool> CVarWebRTCDisableAudioSync;
 	UPROPERTY(config, EditAnywhere, Category = "WebRtcSettings", meta = (
-		ConsoleVariable = "PixelStreaming.WebRTC.DisableAudioSync",
+		ConsoleVariable = "PixelStreaming2.WebRTC.DisableAudioSync",
 		DisplayName = "Disable Audio/Video track sychronisation",
 		ToolTip = "Disables the synchronization of audio and video tracks in WebRTC. This can be useful in low latency usecases where synchronization is not required."
 		))
@@ -354,7 +354,7 @@ public:
 
 	static TAutoConsoleVariable<bool> CVarWebRTCEnableFlexFec;
 	UPROPERTY(config, EditAnywhere, Category = "WebRtcSettings", meta = (
-		ConsoleVariable = "PixelStreaming.WebRTC.EnableFlexFec",
+		ConsoleVariable = "PixelStreaming2.WebRTC.EnableFlexFec",
 		DisplayName = "Enable Flexible Forward Error Correction",
 		ToolTip = "Signals support for Flexible Forward Error Correction to WebRTC. This can cause a reduction in quality if total bitrate is low."
 		))
@@ -362,7 +362,7 @@ public:
 
 	static TAutoConsoleVariable<bool> CVarWebRTCDisableStats;
 	UPROPERTY(config, EditAnywhere, Category = "WebRtcSettings", meta = (
-		ConsoleVariable = "PixelStreaming.WebRTC.DisableStats",
+		ConsoleVariable = "PixelStreaming2.WebRTC.DisableStats",
 		DisplayName = "Disable WebRtc Stats",
 		ToolTip = "Disables the collection of WebRTC stats."
 		))
@@ -370,7 +370,7 @@ public:
 
 	static TAutoConsoleVariable<bool> CVarWebRTCNegotiateCodecs;
 	UPROPERTY(config, EditAnywhere, Category = "WebRtcSettings", meta = (
-		ConsoleVariable = "PixelStreaming.WebRTC.NegotiateCodecs",
+		ConsoleVariable = "PixelStreaming2.WebRTC.NegotiateCodecs",
 		DisplayName = "WebRtc Negotiate Codecs",
 		ToolTip = "Whether PS should send all its codecs during sdp handshake so peers can negotiate or just send a single selected codec."
 		))
@@ -378,7 +378,7 @@ public:
 
 	static TAutoConsoleVariable<float> CVarWebRTCAudioGain;
 	UPROPERTY(config, EditAnywhere, Category = "WebRtcSettings", meta = (
-		ConsoleVariable = "PixelStreaming.WebRTC.AudioGain",
+		ConsoleVariable = "PixelStreaming2.WebRTC.AudioGain",
 		DisplayName = "Default Audio Gain",
 		ToolTip = "Sets the amount of gain to apply to audio. Default: 1.0"
 		))
@@ -388,7 +388,7 @@ public:
 	UPROPERTY(config, EditAnywhere, Category = "WebRtcSettings", meta = (
 		Bitmask,
 		BitmaskEnum = "EPortAllocatorFlags",
-		ConsoleVariable = "PixelStreaming.WebRTC.PortAllocatorFlags",
+		ConsoleVariable = "PixelStreaming2.WebRTC.PortAllocatorFlags",
 		DisplayName = "WebRtc port allocation flags",
 		ToolTip = "Sets the WebRTC port allocator flags. See "
 		))
@@ -398,7 +398,7 @@ public:
 
 	static TAutoConsoleVariable<int> CVarWebRTCMinPort;
 	UPROPERTY(config, EditAnywhere, Category = "WebRtcSettings", meta = (
-		ConsoleVariable = "PixelStreaming.WebRTC.MinPort",
+		ConsoleVariable = "PixelStreaming2.WebRTC.MinPort",
 		DisplayName = "Default WebRtc Min Port",
 		ToolTip = "Sets the minimum usable port for the WebRTC port allocator. Default: 49152"
 		))
@@ -406,7 +406,7 @@ public:
 
 	static TAutoConsoleVariable<int> CVarWebRTCMaxPort;
 	UPROPERTY(config, EditAnywhere, Category = "WebRtcSettings", meta = (
-		ConsoleVariable = "PixelStreaming.WebRTC.MaxPort",
+		ConsoleVariable = "PixelStreaming2.WebRTC.MaxPort",
 		DisplayName = "Default WebRtc Max Port",
 		ToolTip = "Sets the maximum usable port for the WebRTC port allocator. Default: 65535"
 		))
@@ -414,7 +414,7 @@ public:
 
 	static TAutoConsoleVariable<FString> CVarWebRTCFieldTrials;
 	UPROPERTY(config, EditAnywhere, Category = "WebRtcSettings|Field Trials", meta = (
-		ConsoleVariable = "PixelStreaming.WebRTC.FieldTrials",
+		ConsoleVariable = "PixelStreaming2.WebRTC.FieldTrials",
 		DisplayName = "Manually defined WebRtc Field Trials",
 		ToolTip = "Sets the WebRTC field trials string. Format:\"TRIAL1/VALUE1/TRIAL2/VALUE2/\" see https://webrtc.googlesource.com/src/+/HEAD/g3doc/field-trials.md"
 		))
@@ -422,7 +422,7 @@ public:
 
 	static TAutoConsoleVariable<bool> CVarWebRTCDisableFrameDropper;
 	UPROPERTY(config, EditAnywhere, Category = "WebRtcSettings|Field Trials", meta = (
-		ConsoleVariable = "PixelStreaming.WebRTC.DisableFrameDropper",
+		ConsoleVariable = "PixelStreaming2.WebRTC.DisableFrameDropper",
 		DisplayName = "Disable Frame Dropper",
 		ToolTip = "Disables the WebRTC internal frame dropper using the field trial WebRTC-FrameDropper/Disabled/"
 		))
@@ -430,7 +430,7 @@ public:
 
 	static TAutoConsoleVariable<float> CVarWebRTCVideoPacingMaxDelay;
 	UPROPERTY(config, EditAnywhere, Category = "WebRtcSettings|Field Trials", meta = (
-		ConsoleVariable = "PixelStreaming.WebRTC.VideoPacing.MaxDelay",
+		ConsoleVariable = "PixelStreaming2.WebRTC.VideoPacing.MaxDelay",
 		DisplayName = "Video Pacing Max Delay",
 		ToolTip = "Enables the WebRTC-Video-Pacing field trial and sets the max delay (ms) parameter. Default: -1.0f (values below zero are discarded.)"
 		))
@@ -438,7 +438,7 @@ public:
 
 	static TAutoConsoleVariable<float> CVarWebRTCVideoPacingFactor;
 	UPROPERTY(config, EditAnywhere, Category = "WebRtcSettings|Field Trials", meta = (
-		ConsoleVariable = "PixelStreaming.WebRTC.VideoPacing.Factor",
+		ConsoleVariable = "PixelStreaming2.WebRTC.VideoPacing.Factor",
 		DisplayName = "Video Pacing Factor",
 		ToolTip = "Enables the WebRTC-Video-Pacing field trial and sets the video pacing factor parameter. Larger values are more lenient on larger bitrates. Default: -1.0f (values below zero are discarded.)"
 		))
@@ -450,7 +450,7 @@ public:
 	// Begin EditorStreaming CVars
 	static TAutoConsoleVariable<bool> CVarEditorStartOnLaunch;
 	UPROPERTY(config, EditAnywhere, Category = "EditorStreaming", meta = (
-		ConsoleVariable = "PixelStreaming.Editor.StartOnLaunch",
+		ConsoleVariable = "PixelStreaming2.Editor.StartOnLaunch",
 		DisplayName = "Start On Launch",
 		ToolTip = "Start Editor Streaming as soon as the Unreal Editor is launched. Default: false"
 		))
@@ -458,7 +458,7 @@ public:
 
 	static TAutoConsoleVariable<bool> CVarEditorUseRemoteSignallingServer;
 	UPROPERTY(config, EditAnywhere, Category = "EditorStreaming", meta = (
-		ConsoleVariable = "PixelStreaming.Editor.UseRemoteSignallingServer",
+		ConsoleVariable = "PixelStreaming2.Editor.UseRemoteSignallingServer",
 		DisplayName = "Use Remote Signalling Server",
 		ToolTip = "Enables the use of a remote signalling server. Default: false"
 		))
@@ -466,7 +466,7 @@ public:
 	
     static TAutoConsoleVariable<FString> CVarEditorSource;
 	UPROPERTY(config, EditAnywhere, Category = "EditorStreaming", meta = (
-		MappedConsoleVariable = "PixelStreaming.Editor.Source",
+		MappedConsoleVariable = "PixelStreaming2.Editor.Source",
 		DisplayName = "Source",
 		ToolTip = "Editor PixelStreaming source. Supported values are `Editor`, `LevelEditorViewport`. Default: Editor"
 		))
@@ -476,7 +476,7 @@ public:
 	// Begin HMD CVars
 	static TAutoConsoleVariable<bool>  CVarHMDEnable;
 	UPROPERTY(config, EditAnywhere, Category = "XR Streaming", meta = (
-		ConsoleVariable = "PixelStreaming.HMD.Enable",
+		ConsoleVariable = "PixelStreaming2.HMD.Enable",
 		DisplayName = "Enable HMD",
 		ToolTip = "Enables HMD specific functionality for Pixel Streaming. Namely input handling and stereoscopic rendering. Default: false"
 		))
@@ -484,7 +484,7 @@ public:
 
 	static TAutoConsoleVariable<bool>  CVarHMDMatchAspectRatio;
 	UPROPERTY(config, EditAnywhere, Category = "XR Streaming", meta = (
-		ConsoleVariable = "PixelStreaming.HMD.MatchAspectRatio",
+		ConsoleVariable = "PixelStreaming2.HMD.MatchAspectRatio",
 		DisplayName = "Match Aspect Ratio",
 		ToolTip = "If true automatically resize the rendering resolution to match the aspect ratio determined by the HFoV and VFoV. Default: true"
 		))
@@ -492,7 +492,7 @@ public:
 	
 	static TAutoConsoleVariable<bool>  CVarHMDApplyEyePosition;
 	UPROPERTY(config, EditAnywhere, Category = "XR Streaming", meta = (
-		ConsoleVariable = "PixelStreaming.HMD.ApplyEyePosition",
+		ConsoleVariable = "PixelStreaming2.HMD.ApplyEyePosition",
 		DisplayName = "Apply Eye Position",
 		ToolTip = "If true automatically position each eye's rendering by whatever amount WebXR reports for each left-right XRView. If false do no eye positioning. Default: true"
 		))
@@ -500,7 +500,7 @@ public:
 
 	static TAutoConsoleVariable<bool>  CVarHMDApplyEyeRotation;
 	UPROPERTY(config, EditAnywhere, Category = "XR Streaming", meta = (
-		ConsoleVariable = "PixelStreaming.HMD.ApplyEyeRotation",
+		ConsoleVariable = "PixelStreaming2.HMD.ApplyEyeRotation",
 		DisplayName = "Apply Eye Position",
 		ToolTip = "If true automatically rotate each eye's rendering by whatever amount WebXR reports for each left-right XRView. If false do no eye rotation. Default: true"
 		))
@@ -508,7 +508,7 @@ public:
 
 	static TAutoConsoleVariable<float> CVarHMDHFOV;
 	UPROPERTY(config, EditAnywhere, Category = "XR Streaming", meta = (
-		ConsoleVariable = "PixelStreaming.HMD.HFOV",
+		ConsoleVariable = "PixelStreaming2.HMD.HFOV",
 		DisplayName = "Horizontal FOV Override",
 		ToolTip = "Overrides the horizontal field of view for HMD rendering, values are in degrees and values less than 0.0f disable the override. Default: -1.0f"
 		))
@@ -516,7 +516,7 @@ public:
 
 	static TAutoConsoleVariable<float> CVarHMDVFOV;
 	UPROPERTY(config, EditAnywhere, Category = "XR Streaming", meta = (
-		ConsoleVariable = "PixelStreaming.HMD.VFOV",
+		ConsoleVariable = "PixelStreaming2.HMD.VFOV",
 		DisplayName = "Vertical FOV Override",
 		ToolTip = "Overrides the vertical field of view for HMD rendering, values are in degrees and values less than 0.0f disable the override. Default: -1.0f"
 		))
@@ -524,7 +524,7 @@ public:
 
 	static TAutoConsoleVariable<float> CVarHMDIPD;
 	UPROPERTY(config, EditAnywhere, Category = "XR Streaming", meta = (
-		ConsoleVariable = "PixelStreaming.HMD.IPD",
+		ConsoleVariable = "PixelStreaming2.HMD.IPD",
 		DisplayName = "Interpupillary Distance Override",
 		ToolTip = "Overrides the HMD IPD (interpupillary distance), values are in centimeters and values less than 0.0f disable the override. Default: -1.0f"
 		))
@@ -532,7 +532,7 @@ public:
 
 	static TAutoConsoleVariable<float> CVarHMDProjectionOffsetX;
 	UPROPERTY(config, EditAnywhere, Category = "XR Streaming", meta = (
-		ConsoleVariable = "PixelStreaming.HMD.ProjectionOffsetX",
+		ConsoleVariable = "PixelStreaming2.HMD.ProjectionOffsetX",
 		DisplayName = "Horizontal Projection Offset Override",
 		ToolTip = "Overrides the left/right eye projection matrix x-offset, values are in clip space and values less than 0.0f disable the override. Default: -1.0f"
 		))
@@ -540,7 +540,7 @@ public:
 
 	static TAutoConsoleVariable<float> CVarHMDProjectionOffsetY;
 	UPROPERTY(config, EditAnywhere, Category = "XR Streaming", meta = (
-		ConsoleVariable = "PixelStreaming.HMD.ProjectionOffsetY",
+		ConsoleVariable = "PixelStreaming2.HMD.ProjectionOffsetY",
 		DisplayName = "Vertical Projection Offset Override",
 		ToolTip = "Overrides the left-right eye projection matrix y-offset, values are in clip space and values less than 0.0f disable the override. Default: -1.0f"
 		))
@@ -550,7 +550,7 @@ public:
 	// Begin Input CVars
 	static TAutoConsoleVariable<bool> CVarInputAllowConsoleCommands;
 	UPROPERTY(config, EditAnywhere, Category = "Input", meta = (
-		ConsoleVariable = "PixelStreaming.AllowPixelStreaming2Commands",
+		ConsoleVariable = "PixelStreaming2.AllowPixelStreaming2Commands",
 		DisplayName = "Allow Commands",
 		ToolTip = "If true browser can send consoleCommand payloads that execute in UE's console. Default: false"
 		))
@@ -558,7 +558,7 @@ public:
 
 	static TAutoConsoleVariable<FString> CVarInputKeyFilter;
 	UPROPERTY(config, EditAnywhere, Category = "Input", meta = (
-		ConsoleVariable = "PixelStreaming.KeyFilter",
+		ConsoleVariable = "PixelStreaming2.KeyFilter",
 		DisplayName = "Key Filter",
 		ToolTip = "Comma separated list of keys to ignore from streaming clients. Default: \"\""
 		))
diff --git a/Engine/Plugins/Media/PixelStreaming2/Source/PixelStreaming2Core/Private/PixelStreaming2PluginSettings.cpp b/Engine/Plugins/Media/PixelStreaming2/Source/PixelStreaming2Core/Private/PixelStreaming2PluginSettings.cpp
--- a/Engine/Plugins/Media/PixelStreaming2/Source/PixelStreaming2Core/Private/PixelStreaming2PluginSettings.cpp
+++ b/Engine/Plugins/Media/PixelStreaming2/Source/PixelStreaming2Core/Private/PixelStreaming2PluginSettings.cpp
@@ -29,9 +29,9 @@ namespace
 
 	static void VerifyCVarVideoSettings(IConsoleVariable* /* We ignore the passed in console variable as this method is called by many different CVars */)
 	{
-		IConsoleVariable* SimulcastCVar = IConsoleManager::Get().FindConsoleVariable(TEXT("PixelStreaming.Encoder.EnableSimulcast"));
-		IConsoleVariable* CodecCVar = IConsoleManager::Get().FindConsoleVariable(TEXT("PixelStreaming.Encoder.Codec"));
-		IConsoleVariable* ScalabilityModeCVar = IConsoleManager::Get().FindConsoleVariable(TEXT("PixelStreaming.Encoder.ScalabilityMode"));
+		IConsoleVariable* SimulcastCVar = IConsoleManager::Get().FindConsoleVariable(TEXT("PixelStreaming2.Encoder.EnableSimulcast"));
+		IConsoleVariable* CodecCVar = IConsoleManager::Get().FindConsoleVariable(TEXT("PixelStreaming2.Encoder.Codec"));
+		IConsoleVariable* ScalabilityModeCVar = IConsoleManager::Get().FindConsoleVariable(TEXT("PixelStreaming2.Encoder.ScalabilityMode"));
 
 		// Verify that the video codec and scalability mode strings correctly map to an enum
 		CheckConsoleEnum<EVideoCodec>(CodecCVar);
@@ -72,68 +72,68 @@ static FName PixelStreaming2MappedConsoleVariableFName(TEXT("MappedConsoleVariab
 // Begin Pixel Streaming Plugin CVars
 
 TAutoConsoleVariable<bool> UPixelStreaming2PluginSettings::CVarOnScreenStats(
-	TEXT("PixelStreaming.HUDStats"),
+	TEXT("PixelStreaming2.HUDStats"),
 	false,
 	TEXT("Whether to show PixelStreaming stats on the in-game HUD (default: false)."),
 	ECVF_Default);
 
 TAutoConsoleVariable<bool> UPixelStreaming2PluginSettings::CVarLogStats(
-	TEXT("PixelStreaming.LogStats"),
+	TEXT("PixelStreaming2.LogStats"),
 	false,
 	TEXT("Whether to show PixelStreaming stats in the log (default: false)."),
 	FConsoleVariableDelegate::CreateLambda([](IConsoleVariable* Var) { Delegates()->OnLogStatsChanged.Broadcast(Var); }),
 	ECVF_Default);
 
 TAutoConsoleVariable<bool> UPixelStreaming2PluginSettings::CVarDisableLatencyTester(
-	TEXT("PixelStreaming.DisableLatencyTester"),
+	TEXT("PixelStreaming2.DisableLatencyTester"),
 	false,
 	TEXT("If true disables latency tester being triggerable."),
 	ECVF_Default);
 
 TAutoConsoleVariable<FString> UPixelStreaming2PluginSettings::CVarInputController(
-	TEXT("PixelStreaming.InputController"),
+	TEXT("PixelStreaming2.InputController"),
 	TEXT("Any"),
 	TEXT("Various modes of input control supported by Pixel Streaming, currently: \"Any\"  or \"Host\". Default: Any"),
 	ECVF_Default);
 
 TAutoConsoleVariable<bool> UPixelStreaming2PluginSettings::CVarDecoupleFramerate(
-	TEXT("PixelStreaming.DecoupleFramerate"),
+	TEXT("PixelStreaming2.DecoupleFramerate"),
 	false,
 	TEXT("Whether we should only stream as fast as we render or at some fixed interval. Coupled means only stream what we render."),
 	ECVF_Default);
 
 TAutoConsoleVariable<float> UPixelStreaming2PluginSettings::CVarDecoupleWaitFactor(
-	TEXT("PixelStreaming.DecoupleWaitFactor"),
+	TEXT("PixelStreaming2.DecoupleWaitFactor"),
 	1.25f,
 	TEXT("Frame rate factor to wait for a captured frame when streaming in decoupled mode. Higher factor waits longer but may also result in higher latency."),
 	ECVF_Default);
 
 TAutoConsoleVariable<float> UPixelStreaming2PluginSettings::CVarSignalingReconnectInterval(
-	TEXT("PixelStreaming.SignalingReconnectInterval"),
+	TEXT("PixelStreaming2.SignalingReconnectInterval"),
 	2.0f,
 	TEXT("Changes the number of seconds between attempted reconnects to the signaling server. This is useful for reducing the log spam produced from attempted reconnects. A value <= 0 results in no reconnect. Default: 2.0s"),
 	ECVF_Default);
 
 TAutoConsoleVariable<float> UPixelStreaming2PluginSettings::CVarSignalingKeepAliveInterval(
-	TEXT("PixelStreaming.SignalingKeepAliveInterval"),
+	TEXT("PixelStreaming2.SignalingKeepAliveInterval"),
 	30.0f,
 	TEXT("Changes the number of seconds between pings to the signaling server. This is useful for keeping the connection active. A value <= 0 results in no pings. Default: 30.0"),
 	ECVF_Default);
 
 TAutoConsoleVariable<bool> UPixelStreaming2PluginSettings::CVarUseMediaCapture(
-	TEXT("PixelStreaming.UseMediaCapture"),
+	TEXT("PixelStreaming2.UseMediaCapture"),
 	true,
 	TEXT("Use Media Capture from MediaIOFramework to capture frames rather than Pixel Streamings internal backbuffer sources."),
 	ECVF_Default);
 
 TAutoConsoleVariable<FString> UPixelStreaming2PluginSettings::CVarDefaultStreamerID(
-	TEXT("PixelStreaming.ID"),
+	TEXT("PixelStreaming2.ID"),
 	TEXT("DefaultStreamer"),
 	TEXT("Default Streamer ID to be used when not specified elsewhere."),
 	ECVF_Default);
 
 TAutoConsoleVariable<FString> UPixelStreaming2PluginSettings::CVarSignallingURL(
-	TEXT("PixelStreaming.SignallingURL"),
+	TEXT("PixelStreaming2.SignallingURL"),
 	TEXT(""),
 	TEXT("Default URL to connect to for signalling."),
 	ECVF_Default);
@@ -141,77 +141,77 @@ TAutoConsoleVariable<FString> UPixelStreaming2PluginSettings::CVarSignallingURL(
 // Begin Encoder CVars
 
 TAutoConsoleVariable<int32> UPixelStreaming2PluginSettings::CVarEncoderTargetBitrate(
-	TEXT("PixelStreaming.Encoder.TargetBitrate"),
+	TEXT("PixelStreaming2.Encoder.TargetBitrate"),
 	-1,
 	TEXT("Target bitrate (bps). Ignore the bitrate WebRTC wants (not recommended). Set to -1 to disable. Default -1."),
 	ECVF_RenderThreadSafe);
 
 TAutoConsoleVariable<int32> UPixelStreaming2PluginSettings::CVarEncoderMinQuality(
-	TEXT("PixelStreaming.Encoder.MinQuality"),
+	TEXT("PixelStreaming2.Encoder.MinQuality"),
 	0,
 	TEXT("0-100, Higher values result in a better minimum quality but higher average bitrates. Default 0 - i.e. no limit on a minimum Quality."),
 	ECVF_Default);
 
 TAutoConsoleVariable<int32> UPixelStreaming2PluginSettings::CVarEncoderMaxQuality(
-	TEXT("PixelStreaming.Encoder.MaxQuality"),
+	TEXT("PixelStreaming2.Encoder.MaxQuality"),
 	100,
 	TEXT("0-100, Lower values result in lower average bitrates but reduces maximum achievable quality. Default 100 - i.e. no limit on a maximum Quality."),
 	ECVF_Default);
 
 TAutoConsoleVariable<FString> UPixelStreaming2PluginSettings::CVarEncoderQualityPreset(
-	TEXT("PixelStreaming.Encoder.QualityPreset"),
+	TEXT("PixelStreaming2.Encoder.QualityPreset"),
 	TEXT("Default"),
 	TEXT("PixelStreaming encoder presets that affecting Quality vs Bitrate. Supported modes are: `ULTRA_LOW_QUALITY`, `LOW_QUALITY`, `DEFAULT`, `HIGH_QUALITY` or `LOSSLESS`"),
 	FConsoleVariableDelegate::CreateStatic(&CheckConsoleEnum<EAVPreset>),
 	ECVF_Default);
 
 TAutoConsoleVariable<FString> UPixelStreaming2PluginSettings::CVarEncoderLatencyMode(
-	TEXT("PixelStreaming.Encoder.LatencyMode"),
+	TEXT("PixelStreaming2.Encoder.LatencyMode"),
 	TEXT("UltraLowLatency"),
 	TEXT("PixelStreaming encoder mode that affecting Quality vs Latency. Supported modes are: `ULTRA_LOW_LATENCY`, `LOW_LATENCY` or `DEFAULT`"),
 	ECVF_Default);
 
 TAutoConsoleVariable<int32> UPixelStreaming2PluginSettings::CVarEncoderKeyframeInterval(
-	TEXT("PixelStreaming.Encoder.KeyframeInterval"),
+	TEXT("PixelStreaming2.Encoder.KeyframeInterval"),
 	-1,
 	TEXT("How many frames before a key frame is sent. Default: -1 which disables the sending of periodic key frames. Note: NVENC reqires a reinitialization when this changes."),
 	ECVF_Default);
 
 TAutoConsoleVariable<int32> UPixelStreaming2PluginSettings::CVarEncoderMaxSessions(
-	TEXT("PixelStreaming.Encoder.MaxSessions"),
+	TEXT("PixelStreaming2.Encoder.MaxSessions"),
 	-1,
 	TEXT("-1 implies no limit. Maximum number of concurrent hardware encoder sessions for Pixel Streaming. Note GeForce gpus only support 8 concurrent sessions and will rollover to software encoding when that number is exceeded."),
 	ECVF_Default);
 
 TAutoConsoleVariable<bool> UPixelStreaming2PluginSettings::CVarEncoderEnableSimulcast(
-	TEXT("PixelStreaming.Encoder.EnableSimulcast"),
+	TEXT("PixelStreaming2.Encoder.EnableSimulcast"),
 	false,
 	TEXT("Enables simulcast. When enabled, the encoder will encode at full resolution, 1/2 resolution and 1/4 resolution simultaneously. Note: Simulcast is only supported with `H264` and `VP8` and you must use the SFU from the infrastructure to fully utilise this functionality."),
 	FConsoleVariableDelegate::CreateLambda([](IConsoleVariable* Var) { VerifyCVarVideoSettings(nullptr); Delegates()->OnSimulcastEnabledChanged.Broadcast(Var); }),
 	ECVF_Default);
 
 TAutoConsoleVariable<FString> UPixelStreaming2PluginSettings::CVarEncoderCodec(
-	TEXT("PixelStreaming.Encoder.Codec"),
+	TEXT("PixelStreaming2.Encoder.Codec"),
 	"H264",
 	TEXT("PixelStreaming default encoder codec. Supported values are: `H264`, `VP8`, `VP9` or `AV1`"),
 	FConsoleVariableDelegate::CreateStatic(&VerifyCVarVideoSettings),
 	ECVF_Default);
 
 TAutoConsoleVariable<FString> UPixelStreaming2PluginSettings::CVarEncoderScalabilityMode(
-	TEXT("PixelStreaming.Encoder.ScalabilityMode"),
+	TEXT("PixelStreaming2.Encoder.ScalabilityMode"),
 	TEXT("L1T1"),
 	TEXT("Indicates number of Spatial and temporal layers used, default: L1T1. For a full list of values refer to https://www.w3.org/TR/webrtc-svc/#scalabilitymodes*"),
 	FConsoleVariableDelegate::CreateLambda([](IConsoleVariable* Var) { VerifyCVarVideoSettings(nullptr); Delegates()->OnScalabilityModeChanged.Broadcast(Var); }),
 	ECVF_Default);
 
 TAutoConsoleVariable<FString> UPixelStreaming2PluginSettings::CVarEncoderH264Profile(
-	TEXT("PixelStreaming.Encoder.H264Profile"),
+	TEXT("PixelStreaming2.Encoder.H264Profile"),
 	TEXT("Baseline"),
 	TEXT("PixelStreaming encoder profile. Supported modes are: `AUTO`, `BASELINE`, `MAIN`, `HIGH`, `PROGRESSIVE_HIGH`, `CONSTRAINED_HIGH` or `HIGH444`"),
 	ECVF_Default);
 
 TAutoConsoleVariable<bool> UPixelStreaming2PluginSettings::CVarEncoderDebugDumpFrame(
-	TEXT("PixelStreaming.Encoder.DumpDebugFrames"),
+	TEXT("PixelStreaming2.Encoder.DumpDebugFrames"),
 	false,
 	TEXT("Dumps frames from the encoder to a file on disk for debugging purposes."),
 	ECVF_Default);
@@ -219,7 +219,7 @@ TAutoConsoleVariable<bool> UPixelStreaming2PluginSettings::CVarEncoderDebugDumpF
 // Begin WebRTC CVars
 
 TAutoConsoleVariable<int32> UPixelStreaming2PluginSettings::CVarWebRTCFps(
-	TEXT("PixelStreaming.WebRTC.Fps"),
+	TEXT("PixelStreaming2.WebRTC.Fps"),
 	60,
 	TEXT("Framerate for WebRTC encoding. Default: 60"),
 	FConsoleVariableDelegate::CreateLambda([](IConsoleVariable* Var) { Delegates()->OnWebRTCFpsChanged.Broadcast(Var); }),
@@ -227,75 +227,75 @@ TAutoConsoleVariable<int32> UPixelStreaming2PluginSettings::CVarWebRTCFps(
 
 // Note: 1 megabit is the maximum allowed in WebRTC for a start bitrate.
 TAutoConsoleVariable<int32> UPixelStreaming2PluginSettings::CVarWebRTCStartBitrate(
-	TEXT("PixelStreaming.WebRTC.StartBitrate"),
+	TEXT("PixelStreaming2.WebRTC.StartBitrate"),
 	1000000,
 	TEXT("Start bitrate (bps) that WebRTC will try begin the stream with. Must be between Min/Max bitrates. Default: 1000000"),
 	ECVF_RenderThreadSafe);
 
 TAutoConsoleVariable<int32> UPixelStreaming2PluginSettings::CVarWebRTCMinBitrate(
-	TEXT("PixelStreaming.WebRTC.MinBitrate"),
+	TEXT("PixelStreaming2.WebRTC.MinBitrate"),
 	100000,
 	TEXT("Min bitrate (bps) that WebRTC will not request below. Careful not to set too high otherwise WebRTC will just drop frames. Default: 100000"),
 	FConsoleVariableDelegate::CreateLambda([](IConsoleVariable* Var) { Delegates()->OnWebRTCBitrateChanged.Broadcast(Var); }),
 	ECVF_RenderThreadSafe);
 
 TAutoConsoleVariable<int32> UPixelStreaming2PluginSettings::CVarWebRTCMaxBitrate(
-	TEXT("PixelStreaming.WebRTC.MaxBitrate"),
+	TEXT("PixelStreaming2.WebRTC.MaxBitrate"),
 	40000000,
 	TEXT("Max bitrate (bps) that WebRTC will not request above. Default: 40000000 aka 40 megabits/per second."),
 	FConsoleVariableDelegate::CreateLambda([](IConsoleVariable* Var) { Delegates()->OnWebRTCBitrateChanged.Broadcast(Var); }),
 	ECVF_RenderThreadSafe);
 
 TAutoConsoleVariable<bool> UPixelStreaming2PluginSettings::CVarWebRTCDisableReceiveAudio(
-	TEXT("PixelStreaming.WebRTC.DisableReceiveAudio"),
+	TEXT("PixelStreaming2.WebRTC.DisableReceiveAudio"),
 	false,
 	TEXT("Disables receiving audio from the browser into UE."),
 	ECVF_Default);
 
 TAutoConsoleVariable<bool> UPixelStreaming2PluginSettings::CVarWebRTCDisableReceiveVideo(
-	TEXT("PixelStreaming.WebRTC.DisableReceiveVideo"),
+	TEXT("PixelStreaming2.WebRTC.DisableReceiveVideo"),
 	true,
 	TEXT("Disables receiving video from the browser into UE."),
 	ECVF_Default);
 
 TAutoConsoleVariable<bool> UPixelStreaming2PluginSettings::CVarWebRTCDisableTransmitAudio(
-	TEXT("PixelStreaming.WebRTC.DisableTransmitAudio"),
+	TEXT("PixelStreaming2.WebRTC.DisableTransmitAudio"),
 	false,
 	TEXT("Disables transmission of UE audio to the browser."),
 	ECVF_Default);
 
 TAutoConsoleVariable<bool> UPixelStreaming2PluginSettings::CVarWebRTCDisableTransmitVideo(
-	TEXT("PixelStreaming.WebRTC.DisableTransmitVideo"),
+	TEXT("PixelStreaming2.WebRTC.DisableTransmitVideo"),
 	false,
 	TEXT("Disables transmission of UE video to the browser."),
 	ECVF_Default);
 
 TAutoConsoleVariable<bool> UPixelStreaming2PluginSettings::CVarWebRTCDisableAudioSync(
-	TEXT("PixelStreaming.WebRTC.DisableAudioSync"),
+	TEXT("PixelStreaming2.WebRTC.DisableAudioSync"),
 	true,
 	TEXT("Disables the synchronization of audio and video tracks in WebRTC. This can be useful in low latency usecases where synchronization is not required."),
 	ECVF_Default);
 
 TAutoConsoleVariable<bool> UPixelStreaming2PluginSettings::CVarWebRTCEnableFlexFec(
-	TEXT("PixelStreaming.WebRTC.EnableFlexFec"),
+	TEXT("PixelStreaming2.WebRTC.EnableFlexFec"),
 	false,
 	TEXT("Signals support for Flexible Forward Error Correction to WebRTC. This can cause a reduction in quality if total bitrate is low."),
 	ECVF_Default);
 
 TAutoConsoleVariable<bool> UPixelStreaming2PluginSettings::CVarWebRTCDisableStats(
-	TEXT("PixelStreaming.WebRTC.DisableStats"),
+	TEXT("PixelStreaming2.WebRTC.DisableStats"),
 	false,
 	TEXT("Disables the collection of WebRTC stats."),
 	ECVF_Default);
 
 TAutoConsoleVariable<bool> UPixelStreaming2PluginSettings::CVarWebRTCNegotiateCodecs(
-	TEXT("PixelStreaming.WebRTC.NegotiateCodecs"),
+	TEXT("PixelStreaming2.WebRTC.NegotiateCodecs"),
 	false,
 	TEXT("Whether PS should send all its codecs during sdp handshake so peers can negotiate or just send a single selected codec."),
 	ECVF_Default);
 
 TAutoConsoleVariable<float> UPixelStreaming2PluginSettings::CVarWebRTCAudioGain(
-	TEXT("PixelStreaming.WebRTC.AudioGain"),
+	TEXT("PixelStreaming2.WebRTC.AudioGain"),
 	1.0f,
 	TEXT("Sets the amount of gain to apply to audio. Default: 1.0"),
 	ECVF_Default);
@@ -304,19 +304,19 @@ TAutoConsoleVariable<float> UPixelStreaming2PluginSettings::CVarWebRTCAudioGain(
 
 // Begin EditorStreaming CVars
 TAutoConsoleVariable<bool> UPixelStreaming2PluginSettings::CVarEditorStartOnLaunch(
-	TEXT("PixelStreaming.Editor.StartOnLaunch"),
+	TEXT("PixelStreaming2.Editor.StartOnLaunch"),
 	false,
 	TEXT("Start Editor Streaming as soon as the Unreal Editor is launched. Default: false"),
 	ECVF_Default);
 
 TAutoConsoleVariable<bool> UPixelStreaming2PluginSettings::CVarEditorUseRemoteSignallingServer(
-	TEXT("PixelStreaming.Editor.UseRemoteSignallingServer"),
+	TEXT("PixelStreaming2.Editor.UseRemoteSignallingServer"),
 	false,
 	TEXT("Enables the use of a remote signalling server. Default: false"),
 	ECVF_Default);
 
 TAutoConsoleVariable<FString> UPixelStreaming2PluginSettings::CVarEditorSource(
-	TEXT("PixelStreaming.Editor.Source"),
+	TEXT("PixelStreaming2.Editor.Source"),
 	TEXT("Editor"),
 	TEXT("Editor PixelStreaming source. Supported values are `Editor`, `LevelEditorViewport`. Default: `Editor`"),
 	FConsoleVariableDelegate::CreateStatic(&CheckConsoleEnum<EPixelStreaming2EditorStreamTypes>),
@@ -325,55 +325,55 @@ TAutoConsoleVariable<FString> UPixelStreaming2PluginSettings::CVarEditorSource(
 
 // Begin HMD CVars
 TAutoConsoleVariable<bool> UPixelStreaming2PluginSettings::CVarHMDEnable(
-	TEXT("PixelStreaming.HMD.Enable"),
+	TEXT("PixelStreaming2.HMD.Enable"),
 	false,
 	TEXT("Enables HMD specific functionality for Pixel Streaming. Namely input handling and stereoscopic rendering. Default: false"),
 	ECVF_Default);
 
 TAutoConsoleVariable<bool> UPixelStreaming2PluginSettings::CVarHMDMatchAspectRatio(
-	TEXT("PixelStreaming.HMD.MatchAspectRatio"),
+	TEXT("PixelStreaming2.HMD.MatchAspectRatio"),
 	true,
 	TEXT("If true automatically resize the rendering resolution to match the aspect ratio determined by the HFoV and VFoV. Default: true"),
 	ECVF_Default);
 
 TAutoConsoleVariable<bool> UPixelStreaming2PluginSettings::CVarHMDApplyEyePosition(
-	TEXT("PixelStreaming.HMD.ApplyEyePosition"),
+	TEXT("PixelStreaming2.HMD.ApplyEyePosition"),
 	true,
 	TEXT("If true automatically position each eye's rendering by whatever amount WebXR reports for each left-right XRView. If false do no eye positioning. Default: true"),
 	ECVF_Default);
 
 TAutoConsoleVariable<bool> UPixelStreaming2PluginSettings::CVarHMDApplyEyeRotation(
-	TEXT("PixelStreaming.HMD.ApplyEyeRotation"),
+	TEXT("PixelStreaming2.HMD.ApplyEyeRotation"),
 	true,
 	TEXT("If true automatically rotate each eye's rendering by whatever amount WebXR reports for each left-right XRView. If false do no eye rotation. Default: true"),
 	ECVF_Default);
 
 TAutoConsoleVariable<float> UPixelStreaming2PluginSettings::CVarHMDHFOV(
-	TEXT("PixelStreaming.HMD.HFOV"),
+	TEXT("PixelStreaming2.HMD.HFOV"),
 	-1.0f,
 	TEXT("Overrides the horizontal field of view for HMD rendering, values are in degrees and values less than 0.0f disable the override. Default: -1.0f"),
 	ECVF_Default);
 
 TAutoConsoleVariable<float> UPixelStreaming2PluginSettings::CVarHMDVFOV(
-	TEXT("PixelStreaming.HMD.VFOV"),
+	TEXT("PixelStreaming2.HMD.VFOV"),
 	-1.0f,
 	TEXT("Overrides the vertical field of view for HMD rendering, values are in degrees and values less than 0.0f disable the override. Default: -1.0f"),
 	ECVF_Default);
 
 TAutoConsoleVariable<float> UPixelStreaming2PluginSettings::CVarHMDIPD(
-	TEXT("PixelStreaming.HMD.IPD"),
+	TEXT("PixelStreaming2.HMD.IPD"),
 	-1.0f,
 	TEXT("Overrides the HMD IPD (interpupillary distance), values are in centimeters and values less than 0.0f disable the override. Default: -1.0f"),
 	ECVF_Default);
 
 TAutoConsoleVariable<float> UPixelStreaming2PluginSettings::CVarHMDProjectionOffsetX(
-	TEXT("PixelStreaming.HMD.ProjectionOffsetX"),
+	TEXT("PixelStreaming2.HMD.ProjectionOffsetX"),
 	-1.0f,
 	TEXT("Overrides the left/right eye projection matrix x-offset, values are in clip space and values less than 0.0f disable the override. Default: -1.0f"),
 	ECVF_Default);
 
 TAutoConsoleVariable<float> UPixelStreaming2PluginSettings::CVarHMDProjectionOffsetY(
-	TEXT("PixelStreaming.HMD.ProjectionOffsetY"),
+	TEXT("PixelStreaming2.HMD.ProjectionOffsetY"),
 	-1.0f,
 	TEXT("Overrides the left-right eye projection matrix y-offset, values are in clip space and values less than 0.0f disable the override. Default: -1.0f"),
 	ECVF_Default);
@@ -381,13 +381,13 @@ TAutoConsoleVariable<float> UPixelStreaming2PluginSettings::CVarHMDProjectionOff
 
 // Begin Input CVars
 TAutoConsoleVariable<bool> UPixelStreaming2PluginSettings::CVarInputAllowConsoleCommands(
-	TEXT("PixelStreaming.AllowPixelStreamingCommands"),
+	TEXT("PixelStreaming2.AllowPixelStreamingCommands"),
 	false,
 	TEXT("If true browser can send consoleCommand payloads that execute in UE's console. Default: false"),
 	ECVF_Default);
 
 TAutoConsoleVariable<FString> UPixelStreaming2PluginSettings::CVarInputKeyFilter(
-	TEXT("PixelStreaming.KeyFilter"),
+	TEXT("PixelStreaming2.KeyFilter"),
 	"",
 	TEXT("Comma separated list of keys to ignore from streaming clients. Default: \"\""),
 	FConsoleVariableDelegate::CreateLambda([](IConsoleVariable* Var) { Delegates()->OnInputKeyFilterChanged.Broadcast(Var); }),
@@ -559,43 +559,43 @@ void SetPortAllocationCVarFromProperty(UObject* This, FProperty* Property)
 }
 
 TAutoConsoleVariable<FString> UPixelStreaming2PluginSettings::CVarWebRTCPortAllocatorFlags(
-	TEXT("PixelStreaming.WebRTC.PortAllocatorFlags"),
+	TEXT("PixelStreaming2.WebRTC.PortAllocatorFlags"),
 	TEXT(""),
 	TEXT("Sets the WebRTC port allocator flags. Format:\"DISABLE_UDP,DISABLE_STUN,...\""),
 	ECVF_Default);
 
 TAutoConsoleVariable<int> UPixelStreaming2PluginSettings::CVarWebRTCMinPort(
-	TEXT("PixelStreaming.WebRTC.MinPort"),
+	TEXT("PixelStreaming2.WebRTC.MinPort"),
 	49152, // Default according to RFC5766
 	TEXT("Sets the minimum usable port for the WebRTC port allocator. Default: 49152"),
 	ECVF_Default);
 
 TAutoConsoleVariable<int> UPixelStreaming2PluginSettings::CVarWebRTCMaxPort(
-	TEXT("PixelStreaming.WebRTC.MaxPort"),
+	TEXT("PixelStreaming2.WebRTC.MaxPort"),
 	65535, // Default according to RFC5766
 	TEXT("Sets the maximum usable port for the WebRTC port allocator. Default: 65535"),
 	ECVF_Default);
 
 TAutoConsoleVariable<FString> UPixelStreaming2PluginSettings::CVarWebRTCFieldTrials(
-	TEXT("PixelStreaming.WebRTC.FieldTrials"),
+	TEXT("PixelStreaming2.WebRTC.FieldTrials"),
 	TEXT(""),
 	TEXT("Sets the WebRTC field trials string. Format:\"TRIAL1/VALUE1/TRIAL2/VALUE2/\""),
 	ECVF_Default);
 
 TAutoConsoleVariable<bool> UPixelStreaming2PluginSettings::CVarWebRTCDisableFrameDropper(
-	TEXT("PixelStreaming.WebRTC.DisableFrameDropper"),
+	TEXT("PixelStreaming2.WebRTC.DisableFrameDropper"),
 	false,
 	TEXT("Disables the WebRTC internal frame dropper using the field trial WebRTC-FrameDropper/Disabled/"),
 	ECVF_Default);
 
 TAutoConsoleVariable<float> UPixelStreaming2PluginSettings::CVarWebRTCVideoPacingMaxDelay(
-	TEXT("PixelStreaming.WebRTC.VideoPacing.MaxDelay"),
+	TEXT("PixelStreaming2.WebRTC.VideoPacing.MaxDelay"),
 	-1.0f,
 	TEXT("Enables the WebRTC-Video-Pacing field trial and sets the max delay (ms) parameter. Default: -1.0f (values below zero are discarded.)"),
 	ECVF_Default);
 
 TAutoConsoleVariable<float> UPixelStreaming2PluginSettings::CVarWebRTCVideoPacingFactor(
-	TEXT("PixelStreaming.WebRTC.VideoPacing.Factor"),
+	TEXT("PixelStreaming2.WebRTC.VideoPacing.Factor"),
 	-1.0f,
 	TEXT("Enables the WebRTC-Video-Pacing field trial and sets the video pacing factor parameter. Larger values are more lenient on larger bitrates. Default: -1.0f (values below zero are discarded.)"),
 	ECVF_Default);
@@ -790,73 +790,73 @@ void UPixelStreaming2PluginSettings::SetCVarFromProperty(IConsoleVariable* CVar,
 
 // Map of Property Names to their commandline args as GetMetaData() is not avaliable in packaged projects
 static const TMap<FString, FString> GetCmdArg = {
-	{ "OnScreenStats", "PixelStreaming.HUDStats" },
-	{ "LogStats", "PixelStreaming.LogStats" },
-	{ "SendPlayerIdAsInteger", "PixelStreaming.SendPlayerIdAsInteger" },
-	{ "DisableLatencyTester", "PixelStreaming.DisableLatencyTester" },
-	{ "DecoupleFramerate", "PixelStreaming.DecoupleFrameRate" },
-	{ "DecoupleWaitFactor", "PixelStreaming.DecoupleWaitFactor" },
-	{ "SignalingReconnectInterval", "PixelStreaming.SignalingReconnectInterval" },
-	{ "SignalingKeepAliveInterval", "PixelStreaming.SignalingKeepAliveInterval" },
-	{ "UseMediaCapture", "PixelStreaming.UseMediaCapture" },
-	{ "DefaultStreamerID", "PixelStreaming.ID" },
-	{ "SignallingURL", "PixelStreaming.SignallingURL" },
-	{ "Codec", "PixelStreaming.Encoder.Codec" },
-	{ "EncoderTargetBitrate", "PixelStreaming.Encoder.TargetBitrate" },
-	{ "EncoderMinQuality", "PixelStreaming.Encoder.MinQuality" },
-	{ "EncoderMaxQuality", "PixelStreaming.Encoder.MaxQuality" },
-	{ "ScalabilityMode", "PixelStreaming.Encoder.ScalabilityMode" },
-	{ "KeyframeInterval", "PixelStreaming.Encoder.KeyframeInterval" },
-	{ "MaxSessions", "PixelStreaming.Encoder.MaxSessions" },
-	{ "WebRTCFps", "PixelStreaming.WebRTC.Fps" },
-	{ "WebRTCStartBitrate", "PixelStreaming.WebRTC.StartBitrate" },
-	{ "WebRTCMinBitrate", "PixelStreaming.WebRTC.MinBitrate" },
-	{ "WebRTCMaxBitrate", "PixelStreaming.WebRTC.MaxBitrate" },
-	{ "WebRTCDisableReceiveAudio", "PixelStreaming.WebRTC.DisableReceiveAudio" },
-	{ "WebRTCDisableReceiveVideo", "PixelStreaming.WebRTC.DisableReceiveVideo" },
-	{ "WebRTCDisableTransmitAudio", "PixelStreaming.WebRTC.DisableTransmitAudio" },
-	{ "WebRTCDisableTransmitVideo", "PixelStreaming.WebRTC.DisableTransmitVideo" },
-	{ "WebRTCDisableAudioSync", "PixelStreaming.WebRTC.DisableAudioSync" },
-	{ "WebRTCEnableFlexFec", "PixelStreaming.WebRTC.EnableFlexFec" },
-	{ "WebRTCDisableStats", "PixelStreaming.WebRTC.DisableStats" },
-	{ "WebRTCNegotiateCodecs", "PixelStreaming.WebRTC.NegotiateCodecs" },
-	{ "WebRTCAudioGain", "PixelStreaming.WebRTC.AudioGain" },
-	{ "WebRTCPortAllocatorFlags", "PixelStreaming.WebRTC.PortAllocatorFlags" },
-	{ "WebRTCMinPort", "PixelStreaming.WebRTC.MinPort" },
-	{ "WebRTCMaxPort", "PixelStreaming.WebRTC.MaxPort" },
-	{ "WebRTCFieldTrials", "PixelStreaming.WebRTC.FieldTrials" },
-	{ "WebRTCDisableFrameDropper", "PixelStreaming.WebRTC.DisableFrameDropper" },
-	{ "WebRTCVideoPacingMaxDelay", "PixelStreaming.WebRTC.VideoPacing.MaxDelay" },
-	{ "WebRTCVideoPacingFactor", "PixelStreaming.WebRTC.VideoPacing.Factor" },
-	{ "EditorStartOnLaunch", "PixelStreaming.Editor.StartOnLaunch" },
-	{ "EditorUseRemoteSignallingServer", "PixelStreaming.Editor.UseRemoteSignallingServer" },
-	{ "HMDEnable", "PixelStreaming.HMD.Enable" },
-	{ "HMDMatchAspectRatio", "PixelStreaming.HMD.MatchAspectRatio" },
-	{ "HMDAppleEyePosition", "PixelStreaming.HMD.ApplyEyePosition" },
-	{ "HMDApplyEyeRotation", "PixelStreaming.HMD.ApplyEyeRotation" },
-	{ "HMDHFOV", "PixelStreaming.HMD.HFOV" },
-	{ "HMDVFOV", "PixelStreaming.HMD.VFOV" },
-	{ "HMDIPD", "PixelStreaming.HMD.IPD" },
-	{ "HMDProjectionOffsetX", "PixelStreaming.HMD.ProjectionOffsetX" },
-	{ "HMDProjectionOffsetY", "PixelStreaming.HMD.ProjectionOffsetY" },
-	{ "InputAllowConsoleCommands", "PixelStreaming.AllowPixelStreamingCommands" },
-	{ "InputKeyFilter", "PixelStreaming.KeyFilter" }
+	{ "OnScreenStats", "PixelStreaming2.HUDStats" },
+	{ "LogStats", "PixelStreaming2.LogStats" },
+	{ "SendPlayerIdAsInteger", "PixelStreaming2.SendPlayerIdAsInteger" },
+	{ "DisableLatencyTester", "PixelStreaming2.DisableLatencyTester" },
+	{ "DecoupleFramerate", "PixelStreaming2.DecoupleFrameRate" },
+	{ "DecoupleWaitFactor", "PixelStreaming2.DecoupleWaitFactor" },
+	{ "SignalingReconnectInterval", "PixelStreaming2.SignalingReconnectInterval" },
+	{ "SignalingKeepAliveInterval", "PixelStreaming2.SignalingKeepAliveInterval" },
+	{ "UseMediaCapture", "PixelStreaming2.UseMediaCapture" },
+	{ "DefaultStreamerID", "PixelStreaming2.ID" },
+	{ "SignallingURL", "PixelStreaming2.SignallingURL" },
+	{ "Codec", "PixelStreaming2.Encoder.Codec" },
+	{ "EncoderTargetBitrate", "PixelStreaming2.Encoder.TargetBitrate" },
+	{ "EncoderMinQuality", "PixelStreaming2.Encoder.MinQuality" },
+	{ "EncoderMaxQuality", "PixelStreaming2.Encoder.MaxQuality" },
+	{ "ScalabilityMode", "PixelStreaming2.Encoder.ScalabilityMode" },
+	{ "KeyframeInterval", "PixelStreaming2.Encoder.KeyframeInterval" },
+	{ "MaxSessions", "PixelStreaming2.Encoder.MaxSessions" },
+	{ "WebRTCFps", "PixelStreaming2.WebRTC.Fps" },
+	{ "WebRTCStartBitrate", "PixelStreaming2.WebRTC.StartBitrate" },
+	{ "WebRTCMinBitrate", "PixelStreaming2.WebRTC.MinBitrate" },
+	{ "WebRTCMaxBitrate", "PixelStreaming2.WebRTC.MaxBitrate" },
+	{ "WebRTCDisableReceiveAudio", "PixelStreaming2.WebRTC.DisableReceiveAudio" },
+	{ "WebRTCDisableReceiveVideo", "PixelStreaming2.WebRTC.DisableReceiveVideo" },
+	{ "WebRTCDisableTransmitAudio", "PixelStreaming2.WebRTC.DisableTransmitAudio" },
+	{ "WebRTCDisableTransmitVideo", "PixelStreaming2.WebRTC.DisableTransmitVideo" },
+	{ "WebRTCDisableAudioSync", "PixelStreaming2.WebRTC.DisableAudioSync" },
+	{ "WebRTCEnableFlexFec", "PixelStreaming2.WebRTC.EnableFlexFec" },
+	{ "WebRTCDisableStats", "PixelStreaming2.WebRTC.DisableStats" },
+	{ "WebRTCNegotiateCodecs", "PixelStreaming2.WebRTC.NegotiateCodecs" },
+	{ "WebRTCAudioGain", "PixelStreaming2.WebRTC.AudioGain" },
+	{ "WebRTCPortAllocatorFlags", "PixelStreaming2.WebRTC.PortAllocatorFlags" },
+	{ "WebRTCMinPort", "PixelStreaming2.WebRTC.MinPort" },
+	{ "WebRTCMaxPort", "PixelStreaming2.WebRTC.MaxPort" },
+	{ "WebRTCFieldTrials", "PixelStreaming2.WebRTC.FieldTrials" },
+	{ "WebRTCDisableFrameDropper", "PixelStreaming2.WebRTC.DisableFrameDropper" },
+	{ "WebRTCVideoPacingMaxDelay", "PixelStreaming2.WebRTC.VideoPacing.MaxDelay" },
+	{ "WebRTCVideoPacingFactor", "PixelStreaming2.WebRTC.VideoPacing.Factor" },
+	{ "EditorStartOnLaunch", "PixelStreaming2.Editor.StartOnLaunch" },
+	{ "EditorUseRemoteSignallingServer", "PixelStreaming2.Editor.UseRemoteSignallingServer" },
+	{ "HMDEnable", "PixelStreaming2.HMD.Enable" },
+	{ "HMDMatchAspectRatio", "PixelStreaming2.HMD.MatchAspectRatio" },
+	{ "HMDAppleEyePosition", "PixelStreaming2.HMD.ApplyEyePosition" },
+	{ "HMDApplyEyeRotation", "PixelStreaming2.HMD.ApplyEyeRotation" },
+	{ "HMDHFOV", "PixelStreaming2.HMD.HFOV" },
+	{ "HMDVFOV", "PixelStreaming2.HMD.VFOV" },
+	{ "HMDIPD", "PixelStreaming2.HMD.IPD" },
+	{ "HMDProjectionOffsetX", "PixelStreaming2.HMD.ProjectionOffsetX" },
+	{ "HMDProjectionOffsetY", "PixelStreaming2.HMD.ProjectionOffsetY" },
+	{ "InputAllowConsoleCommands", "PixelStreaming2.AllowPixelStreamingCommands" },
+	{ "InputKeyFilter", "PixelStreaming2.KeyFilter" }
 };
 
 static const TMap<FString, FString> GetMappedCmdArg = {
-	{ "InputController", "PixelStreaming.InputController" },
-	{ "QualityPreset", "PixelStreaming.Encoder.QualityPreset" },
-	{ "LatencyMode", "PixelStreaming.Encoder.LatencyMode" },
-	{ "H264Profile", "PixelStreaming.Encoder.H264Profile" },
-	{ "EditorSource", "PixelStreaming.Editor.Source" }
+	{ "InputController", "PixelStreaming2.InputController" },
+	{ "QualityPreset", "PixelStreaming2.Encoder.QualityPreset" },
+	{ "LatencyMode", "PixelStreaming2.Encoder.LatencyMode" },
+	{ "H264Profile", "PixelStreaming2.Encoder.H264Profile" },
+	{ "EditorSource", "PixelStreaming2.Editor.Source" }
 };
 
 static const TArray<FString> GetLegacyCmdArg = {
-	"PixelStreaming.Encoder.MinQp",
-	"PixelStreaming.Encoder.MaxQp",
-	"PixelStreaming.IP",
-	"PixelStreaming.Port",
-	"PixelStreaming.URL",
+	"PixelStreaming2.Encoder.MinQp",
+	"PixelStreaming2.Encoder.MaxQp",
+	"PixelStreaming2.IP",
+	"PixelStreaming2.Port",
+	"PixelStreaming2.URL",
 	"AllowPixelStreamingCommands"
 };
 
