[PS2] QOL: Update command line parsing logic so cvars with a "2" can still be set from old PixelStreaming1 command line args (ie without a two)

diff --git a/Engine/Plugins/Media/PixelStreaming2/Source/PixelStreaming2Core/Private/PixelStreaming2PluginSettings.cpp b/Engine/Plugins/Media/PixelStreaming2/Source/PixelStreaming2Core/Private/PixelStreaming2PluginSettings.cpp
--- a/Engine/Plugins/Media/PixelStreaming2/Source/PixelStreaming2Core/Private/PixelStreaming2PluginSettings.cpp
+++ b/Engine/Plugins/Media/PixelStreaming2/Source/PixelStreaming2Core/Private/PixelStreaming2PluginSettings.cpp
@@ -58,11 +58,17 @@ namespace
 		}
 	}
 
-	FString ConsoleVariableToCommandArg(const FString InCVarName)
+	FString ConsoleVariableToCommandArgValue(const FString InCVarName)
 	{
 		// CVars are . deliminated by section. To get their equivilent commandline arg for parsing
 		// we need to remove the . and add a "="
-		return InCVarName.Replace(TEXT("."), TEXT("")).Append(TEXT("="));
+		return InCVarName.Replace(TEXT("."), TEXT("")).Replace(TEXT("PixelStreaming2"), TEXT("PixelStreaming")).Append(TEXT("="));
+	}
+
+	FString ConsoleVariableToCommandArgParam(const FString InCVarName)
+	{
+		// CVars are . deliminated by section. To get their equivilent commandline arg parameter, we need to to remove the .
+		return InCVarName.Replace(TEXT("."), TEXT("")).Replace(TEXT("PixelStreaming2"), TEXT("PixelStreaming"));
 	}
 } // namespace
 
@@ -700,7 +706,7 @@ void UPixelStreaming2PluginSettings::SetCVarFromPropertyAndValue(IConsoleVariabl
 	if (FByteProperty* ByteProperty = CastField<FByteProperty>(Property); ByteProperty != NULL && ByteProperty->Enum != NULL)
 	{
 		int32 CastValue;
-		FParse::Value(FCommandLine::Get(), *ConsoleVariableToCommandArg(CVarString), CastValue);
+		FParse::Value(FCommandLine::Get(), *ConsoleVariableToCommandArgValue(CVarString), CastValue);
 		CVar->Set(CastValue, ECVF_SetByCommandline);
 	}
 	else if (FEnumProperty* EnumProperty = CastField<FEnumProperty>(Property))
@@ -725,7 +731,7 @@ void UPixelStreaming2PluginSettings::SetCVarFromPropertyAndValue(IConsoleVariabl
 		{
 			CVar->Set(false, ECVF_SetByCommandline);
 		}
-		else if (FParse::Param(FCommandLine::Get(), *CVarString.Replace(TEXT("."), TEXT(""))))
+		else if (FParse::Param(FCommandLine::Get(), *ConsoleVariableToCommandArgParam(CVarString)))
 		{
 			CVar->Set(true, ECVF_SetByCommandline);
 		}
@@ -733,13 +739,13 @@ void UPixelStreaming2PluginSettings::SetCVarFromPropertyAndValue(IConsoleVariabl
 	else if (FIntProperty* IntProperty = CastField<FIntProperty>(Property))
 	{
 		int32 CastValue;
-		FParse::Value(FCommandLine::Get(), *ConsoleVariableToCommandArg(CVarString), CastValue);
+		FParse::Value(FCommandLine::Get(), *ConsoleVariableToCommandArgValue(CVarString), CastValue);
 		CVar->Set(CastValue, ECVF_SetByCommandline);
 	}
 	else if (FFloatProperty* FloatProperty = CastField<FFloatProperty>(Property))
 	{
 		float CastValue;
-		FParse::Value(FCommandLine::Get(), *ConsoleVariableToCommandArg(CVarString), CastValue);
+		FParse::Value(FCommandLine::Get(), *ConsoleVariableToCommandArgValue(CVarString), CastValue);
 		CVar->Set(CastValue, ECVF_SetByCommandline);
 	}
 	else if (FStrProperty* StringProperty = CastField<FStrProperty>(Property))
@@ -885,7 +891,7 @@ void UPixelStreaming2PluginSettings::ValidateCommandLineArgs()
 		bool bValidArg = false;
 		for (const TPair<FString, FString>& Pair : GetCmdArg)
 		{
-			FString ValidCommandLineArg = Pair.Value.Replace(TEXT("."), TEXT(""));
+			FString ValidCommandLineArg = ConsoleVariableToCommandArgParam(Pair.Value);
 			if (CurrentCommandLineArg == ValidCommandLineArg)
 			{
 				bValidArg = true;
@@ -896,7 +902,7 @@ void UPixelStreaming2PluginSettings::ValidateCommandLineArgs()
 		{
 			for (const TPair<FString, FString>& Pair : GetMappedCmdArg)
 			{
-				FString ValidCommandLineArg = Pair.Value.Replace(TEXT("."), TEXT(""));
+				FString ValidCommandLineArg = ConsoleVariableToCommandArgParam(Pair.Value);
 				if (CurrentCommandLineArg == ValidCommandLineArg)
 				{
 					bValidArg = true;
@@ -908,7 +914,7 @@ void UPixelStreaming2PluginSettings::ValidateCommandLineArgs()
 		{
 			for (const FString& LegacyArg : GetLegacyCmdArg)
 			{
-				FString ValidCommandLineArg = LegacyArg.Replace(TEXT("."), TEXT(""));
+				FString ValidCommandLineArg = ConsoleVariableToCommandArgParam(LegacyArg);
 				if (CurrentCommandLineArg == ValidCommandLineArg)
 				{
 					bValidArg = true;
@@ -1036,7 +1042,7 @@ void UPixelStreaming2PluginSettings::PostInitProperties()
 			if (Property->GetNameCPP() == "WebRTCPortAllocatorFlags")
 			{
 				FString ConsoleString;
-				if (FParse::Value(FCommandLine::Get(), *ConsoleVariableToCommandArg(CVarString), ConsoleString))
+				if (FParse::Value(FCommandLine::Get(), *ConsoleVariableToCommandArgValue(CVarString), ConsoleString))
 				{
 					SetCVarFromPropertyAndValue(IConsoleManager::Get().FindConsoleVariable(*CVarString), Property, CVarString, ConsoleString);
 				}
@@ -1051,11 +1057,11 @@ void UPixelStreaming2PluginSettings::PostInitProperties()
 			{
 				// Handle a directly parsable commandline
 				FString ConsoleString;
-				if (FParse::Value(FCommandLine::Get(), *ConsoleVariableToCommandArg(CVarString), ConsoleString))
+				if (FParse::Value(FCommandLine::Get(), *ConsoleVariableToCommandArgValue(CVarString), ConsoleString))
 				{
 					SetCVarFromPropertyAndValue(IConsoleManager::Get().FindConsoleVariable(*CVarString), Property, CVarString, ConsoleString);
 				}
-				else if (FParse::Param(FCommandLine::Get(), *CVarString.Replace(TEXT("."), TEXT(""))))
+				else if (FParse::Param(FCommandLine::Get(), *ConsoleVariableToCommandArgParam(CVarString)))
 				{
 					SetCVarFromPropertyAndValue(IConsoleManager::Get().FindConsoleVariable(*CVarString), Property, CVarString, TEXT("true"));
 				}
@@ -1071,7 +1077,7 @@ void UPixelStreaming2PluginSettings::PostInitProperties()
 		{
 			FString CVarString = GetMappedCmdArg[Property->GetNameCPP()];
 			FString ConsoleString;
-			if (FParse::Value(FCommandLine::Get(), *ConsoleVariableToCommandArg(CVarString), ConsoleString))
+			if (FParse::Value(FCommandLine::Get(), *ConsoleVariableToCommandArgValue(CVarString), ConsoleString))
 			{
 				SetCVarFromPropertyAndValue(IConsoleManager::Get().FindConsoleVariable(*CVarString), Property, CVarString, ConsoleString);
 			}
