[PS2] Add: Reinstate IsSignallingConnected functionality to the IPixelStreaming2Streamer.h

diff --git a/Engine/Plugins/Media/PixelStreaming2/Source/PixelStreaming2/Private/Streamer.cpp b/Engine/Plugins/Media/PixelStreaming2/Source/PixelStreaming2/Private/Streamer.cpp
--- a/Engine/Plugins/Media/PixelStreaming2/Source/PixelStreaming2/Private/Streamer.cpp
+++ b/Engine/Plugins/Media/PixelStreaming2/Source/PixelStreaming2/Private/Streamer.cpp
@@ -461,7 +461,6 @@ namespace UE::PixelStreaming2
 	void FStreamer::ConsumeStats(FString PlayerId, FName StatName, float StatValue)
 	{
 
-		
 		if (StatName != PixelStreaming2StatNames::MeanQPPerSecond)
 		{
 			return;
@@ -606,9 +605,9 @@ namespace UE::PixelStreaming2
 
 	void FStreamer::SendInitialSettings(FString PlayerId) const
 	{
-		const FString				   PixelStreaming2Payload = FString::Printf(TEXT("{ \"AllowPixelStreamingCommands\": %s, \"DisableLatencyTest\": %s }"),
-			 UPixelStreaming2PluginSettings::CVarInputAllowConsoleCommands.GetValueOnAnyThread() ? TEXT("true") : TEXT("false"),
-			 UPixelStreaming2PluginSettings::CVarDisableLatencyTester.GetValueOnAnyThread() ? TEXT("true") : TEXT("false"));
+		const FString PixelStreaming2Payload = FString::Printf(TEXT("{ \"AllowPixelStreamingCommands\": %s, \"DisableLatencyTest\": %s }"),
+			UPixelStreaming2PluginSettings::CVarInputAllowConsoleCommands.GetValueOnAnyThread() ? TEXT("true") : TEXT("false"),
+			UPixelStreaming2PluginSettings::CVarDisableLatencyTester.GetValueOnAnyThread() ? TEXT("true") : TEXT("false"));
 
 		const FString WebRTCPayload = FString::Printf(TEXT("{ \"FPS\": %d, \"MinBitrate\": %d, \"MaxBitrate\": %d }"),
 			UPixelStreaming2PluginSettings::CVarWebRTCFps.GetValueOnAnyThread(),
@@ -772,6 +771,13 @@ namespace UE::PixelStreaming2
 		{
 			case EpicRtcSessionState::Connected:
 			{
+				bSignallingConnected = true;
+				if (UPixelStreaming2Delegates* Delegates = UPixelStreaming2Delegates::Get())
+				{
+					Delegates->OnConnectedToSignallingServer.Broadcast(StreamerId);
+					Delegates->OnConnectedToSignallingServerNative.Broadcast(StreamerId);
+				}
+
 				UE_LOG(LogPixelStreaming2, VeryVerbose, TEXT("FStreamer::OnSessionStateUpdate State=Connected"));
 				EpicRtcBitrate Bitrate = {
 					._minBitrateBps = UPixelStreaming2PluginSettings::CVarWebRTCMinBitrate.GetValueOnAnyThread(),
@@ -832,6 +838,12 @@ namespace UE::PixelStreaming2
 				break;
 			case EpicRtcSessionState::Disconnected:
 				// Do something on `disconnected` session here
+				bSignallingConnected = false;
+				if (UPixelStreaming2Delegates* Delegates = UPixelStreaming2Delegates::Get())
+				{
+					Delegates->OnDisconnectedFromSignallingServer.Broadcast(StreamerId);
+					Delegates->OnDisconnectedFromSignallingServerNative.Broadcast(StreamerId);
+				}
 				UE_LOG(LogPixelStreaming2, VeryVerbose, TEXT("FStreamer::OnSessionStateUpdate State=Disconnected"));
 				RemoveSession(false);
 				StopStreaming();
@@ -839,6 +851,7 @@ namespace UE::PixelStreaming2
 				break;
 			case EpicRtcSessionState::Failed:
 				// Do something on `failed` session here
+				bSignallingConnected = false;
 				UE_LOG(LogPixelStreaming2, VeryVerbose, TEXT("FStreamer::OnSessionStateUpdate State=Failed"));
 				break;
 			case EpicRtcSessionState::Exiting:
@@ -1339,7 +1352,7 @@ namespace UE::PixelStreaming2
 			return;
 		}
 
-		const uint8_t								   Type = DataFrame->Data()[0];
+		const uint8_t							 Type = DataFrame->Data()[0];
 		TSharedPtr<IPixelStreaming2DataProtocol> ToStreamerProtocol = InputHandler->GetToStreamerProtocol();
 
 		if (Type == ToStreamerProtocol->Find(EPixelStreaming2ToStreamerMessage::LatencyTest)->GetID())
diff --git a/Engine/Plugins/Media/PixelStreaming2/Source/PixelStreaming2/Private/Streamer.h b/Engine/Plugins/Media/PixelStreaming2/Source/PixelStreaming2/Private/Streamer.h
--- a/Engine/Plugins/Media/PixelStreaming2/Source/PixelStreaming2/Private/Streamer.h
+++ b/Engine/Plugins/Media/PixelStreaming2/Source/PixelStreaming2/Private/Streamer.h
@@ -30,13 +30,14 @@ namespace UE::PixelStreaming2
 		virtual int32 GetStreamFPS() override;
 		virtual void  SetCoupleFramerate(bool bCouple) override;
 
-		virtual void										  SetVideoProducer(TSharedPtr<IPixelStreaming2VideoProducer> Input) override;
+		virtual void									SetVideoProducer(TSharedPtr<IPixelStreaming2VideoProducer> Input) override;
 		virtual TWeakPtr<IPixelStreaming2VideoProducer> GetVideoProducer() override;
 
 		virtual void	SetSignallingServerURL(const FString& InSignallingServerURL) override;
 		virtual FString GetSignallingServerURL() override;
 
 		virtual FString GetId() override { return StreamerId; };
+		virtual bool	IsSignallingConnected() override { return bSignallingConnected; }
 		virtual void	StartStreaming() override;
 		virtual void	StopStreaming() override;
 		virtual bool	IsStreaming() const override { return bStreamingStarted; }
@@ -104,6 +105,7 @@ namespace UE::PixelStreaming2
 		FString SFUPlayerId = INVALID_PLAYER_ID;
 		FString InputControllingId = INVALID_PLAYER_ID;
 
+		bool bSignallingConnected = false;
 		bool bStreamingStarted = false;
 
 		FPreConnectionEvent	   StreamingPreConnectionEvent;
diff --git a/Engine/Plugins/Media/PixelStreaming2/Source/PixelStreaming2/Public/IPixelStreaming2Streamer.h b/Engine/Plugins/Media/PixelStreaming2/Source/PixelStreaming2/Public/IPixelStreaming2Streamer.h
--- a/Engine/Plugins/Media/PixelStreaming2/Source/PixelStreaming2/Public/IPixelStreaming2Streamer.h
+++ b/Engine/Plugins/Media/PixelStreaming2/Source/PixelStreaming2/Public/IPixelStreaming2Streamer.h
@@ -69,6 +69,12 @@ public:
 	 */
 	virtual FString GetId() = 0;
 
+	/**
+	 * @brief Check if this streamer is currently connected to the signalling mechanism (e.g. websocket for signalling server).
+	 * @return True if the streamer is connected to the signalling mechanism.
+	 */
+	virtual bool IsSignallingConnected() = 0;
+
 	/**
 	 * @brief Start streaming this streamer
 	 */
diff --git a/Engine/Plugins/Media/PixelStreaming2/Source/PixelStreaming2/Public/PixelStreaming2Delegates.h b/Engine/Plugins/Media/PixelStreaming2/Source/PixelStreaming2/Public/PixelStreaming2Delegates.h
--- a/Engine/Plugins/Media/PixelStreaming2/Source/PixelStreaming2/Public/PixelStreaming2Delegates.h
+++ b/Engine/Plugins/Media/PixelStreaming2/Source/PixelStreaming2/Public/PixelStreaming2Delegates.h
@@ -12,6 +12,28 @@ class PIXELSTREAMING2_API UPixelStreaming2Delegates : public UObject
 	GENERATED_BODY()
 
 public:
+	/**
+	 * A connection to the signalling server was made.
+	 */
+	// BP Delegate
+	DECLARE_DYNAMIC_MULTICAST_DELEGATE_OneParam(FConnectedToSignallingServer, FString, StreamerId);
+	UPROPERTY(BlueprintAssignable, Category = "Pixel Streaming Delegates")
+	FConnectedToSignallingServer OnConnectedToSignallingServer;
+	// C++ Delegate
+	DECLARE_MULTICAST_DELEGATE_OneParam(FConnectedToSignallingServerNative, FString /* StreamerId */);
+	FConnectedToSignallingServerNative OnConnectedToSignallingServerNative;
+
+	/**
+	 * A connection to the signalling server was lost.
+	 */
+	// BP Delegate
+	DECLARE_DYNAMIC_MULTICAST_DELEGATE_OneParam(FDisconnectedFromSignallingServer, FString, StreamerId);
+	UPROPERTY(BlueprintAssignable, Category = "Pixel Streaming Delegates")
+	FDisconnectedFromSignallingServer OnDisconnectedFromSignallingServer;
+	// C++ Delegate
+	DECLARE_MULTICAST_DELEGATE_OneParam(FDisconnectedFromSignallingServerNative, FString /* StreamerId */);
+	FDisconnectedFromSignallingServerNative OnDisconnectedFromSignallingServerNative;
+
 	/**
 	 * A new connection has been made to the session.
 	 */
