[PS2] Add: Re-implement kick player functionality

diff --git a/Engine/Plugins/Media/PixelStreaming2/Source/PixelStreaming2/Private/PlayerContext.h b/Engine/Plugins/Media/PixelStreaming2/Source/PixelStreaming2/Private/PlayerContext.h
--- a/Engine/Plugins/Media/PixelStreaming2/Source/PixelStreaming2/Private/PlayerContext.h
+++ b/Engine/Plugins/Media/PixelStreaming2/Source/PixelStreaming2/Private/PlayerContext.h
@@ -13,16 +13,19 @@
 #include "epic_rtc/core/audio/audio_track.h"
 #include "epic_rtc/core/video/video_track.h"
 #include "epic_rtc/core/data_track.h"
+#include "epic_rtc/core/participant.h"
 
 namespace UE::PixelStreaming2
 {
 	struct FPlayerContext
 	{
+		TRefCountPtr<EpicRtcParticipantInterface> ParticipantInterface;
+
 		TSharedPtr<FEpicRtcAudioSource> AudioSource;
 		TSharedPtr<FEpicRtcAudioSink>	AudioSink;
 
-		TSharedPtr<FEpicRtcVideoSource>	  VideoSource;
-		TSharedPtr<FEpicRtcVideoSink> VideoSink;
+		TSharedPtr<FEpicRtcVideoSource> VideoSource;
+		TSharedPtr<FEpicRtcVideoSink>	VideoSink;
 
 		TSharedPtr<FEpicRtcDataTrack> DataTrack;
 	};
diff --git a/Engine/Plugins/Media/PixelStreaming2/Source/PixelStreaming2/Private/Streamer.cpp b/Engine/Plugins/Media/PixelStreaming2/Source/PixelStreaming2/Private/Streamer.cpp
--- a/Engine/Plugins/Media/PixelStreaming2/Source/PixelStreaming2/Private/Streamer.cpp
+++ b/Engine/Plugins/Media/PixelStreaming2/Source/PixelStreaming2/Private/Streamer.cpp
@@ -339,7 +339,10 @@ namespace UE::PixelStreaming2
 
 	void FStreamer::KickPlayer(FString PlayerId)
 	{
-		// TODO (Migration) RTCP-7024 Delete player session?
+		if (FPlayerContext* PlayerContext = Players->Find(PlayerId))
+		{
+			PlayerContext->ParticipantInterface->Kick();
+		}
 	}
 
 	TArray<FString> FStreamer::GetConnectedPlayers()
@@ -908,6 +911,7 @@ namespace UE::PixelStreaming2
 		}
 
 		FPlayerContext& PlayerContext = Players->FindOrAdd(ParticipantId);
+		PlayerContext.ParticipantInterface = Participant;
 
 		TRefCountPtr<EpicRtcConnectionInterface> ParticipantConnection = Participant->GetConnection();
 		ParticipantConnection->SetManualNegotiation(true);
