[PS2] Fix: EpicRtcRoomLifetimeTest

diff --git a/Engine/Plugins/Media/PixelStreaming2/Source/PixelStreaming2/Private/Tests/EpicRtcUtils.cpp b/Engine/Plugins/Media/PixelStreaming2/Source/PixelStreaming2/Private/Tests/EpicRtcUtils.cpp
--- a/Engine/Plugins/Media/PixelStreaming2/Source/PixelStreaming2/Private/Tests/EpicRtcUtils.cpp
+++ b/Engine/Plugins/Media/PixelStreaming2/Source/PixelStreaming2/Private/Tests/EpicRtcUtils.cpp
@@ -6,9 +6,12 @@ namespace UE::PixelStreaming2
 {
 	bool FTickAndWaitOrTimeout::Update()
 	{
-		if (Manager->GetEpicRtcConference() && Manager->GetEpicRtcConference()->NeedsTick())
+		if (Manager->GetEpicRtcConference())
 		{
-			Manager->GetEpicRtcConference()->Tick();
+			while (Manager->GetEpicRtcConference()->NeedsTick())
+			{
+				Manager->GetEpicRtcConference()->Tick();
+			}
 		}
 
 		if (CheckFunc())
@@ -22,19 +25,19 @@ namespace UE::PixelStreaming2
 			UE_LOG(LogPixelStreaming2, Error, TEXT("Timed out"));
 			return true;
 		}
-		return false;	// Latent Test return false will run again next frame
+		return false; // Latent Test return false will run again next frame
 	}
 
 	bool FDisconnectRoom::Update()
 	{
 		TRefCountPtr<EpicRtcRoomInterface>& Room = Manager->GetEpicRtcRoom();
 
-		if(!Room)
+		if (!Room)
 		{
 			UE_LOG(LogPixelStreaming2, Error, TEXT("Unable to disconnect room, Room does not exist"));
 			return true;
 		}
-		
+
 		Room->Leave();
 
 		return true;
@@ -43,12 +46,12 @@ namespace UE::PixelStreaming2
 	bool FCleanupRoom::Update()
 	{
 		TRefCountPtr<EpicRtcRoomInterface>& Room = Manager->GetEpicRtcRoom();
-		if(!Room)
+		if (!Room)
 		{
 			UE_LOG(LogPixelStreaming2, Error, TEXT("Unable to update room, Room does not exist"));
 			return true;
 		}
-		
+
 		Manager->GetEpicRtcSession()->RemoveRoom(ToEpicRtcStringView(RoomId));
 
 		// EpicRtc has released its hold on the room. All that should be holding a ref is the manager
@@ -63,8 +66,8 @@ namespace UE::PixelStreaming2
 		// By directly calling release, we get the final count which we can check.
 		EpicRtcRoomInterface* RoomPtr = Room.GetReference();
 		RoomPtr->AddRef();
-		Room = nullptr;	//will callrelease so count is same as before calling RoomPtr->AddRef();
-		if(uint32 Count = RoomPtr->Release(); Count != 0)
+		Room = nullptr; // will callrelease so count is same as before calling RoomPtr->AddRef();
+		if (uint32 Count = RoomPtr->Release(); Count != 0)
 		{
 			UE_LOG(LogPixelStreaming2, Error, TEXT("Room has invalid reference count. Expected 0, Actual (%d)"), Count);
 		}
