Fix Pixel Streaming 2 Audio Sink race condition.

There was a race condition in Pixel Streaming 2 Audio Sink where removal of a consumer while it is being accessed can result in a crash. Putting a critical section around add, removal and accessing fixes this. This is simialr to what Video Sink does.

diff --git a/Engine/Plugins/Media/PixelStreaming2/Source/PixelStreaming2/Internal/AudioSink.h b/Engine/Plugins/Media/PixelStreaming2/Source/PixelStreaming2/Internal/AudioSink.h
--- a/Engine/Plugins/Media/PixelStreaming2/Source/PixelStreaming2/Internal/AudioSink.h
+++ b/Engine/Plugins/Media/PixelStreaming2/Source/PixelStreaming2/Internal/AudioSink.h
@@ -25,6 +25,7 @@ namespace UE::PixelStreaming2
 		TSet<IPixelStreaming2AudioConsumer*> AudioConsumers;
 
 	private:
+		FCriticalSection AudioConsumersCS;
 		bool bIsMuted = false;
 	};
 } // namespace UE::PixelStreaming2
\ No newline at end of file
diff --git a/Engine/Plugins/Media/PixelStreaming2/Source/PixelStreaming2/Private/AudioSink.cpp b/Engine/Plugins/Media/PixelStreaming2/Source/PixelStreaming2/Private/AudioSink.cpp
--- a/Engine/Plugins/Media/PixelStreaming2/Source/PixelStreaming2/Private/AudioSink.cpp
+++ b/Engine/Plugins/Media/PixelStreaming2/Source/PixelStreaming2/Private/AudioSink.cpp
@@ -10,6 +10,7 @@ namespace UE::PixelStreaming2
 
 	FAudioSink::~FAudioSink()
 	{
+		FScopeLock Lock(&AudioConsumersCS);
 		for (auto Iter = AudioConsumers.CreateIterator(); Iter; ++Iter)
 		{
 			IPixelStreaming2AudioConsumer* AudioConsumer = Iter.ElementIt->Value;
@@ -23,6 +24,7 @@ namespace UE::PixelStreaming2
 
 	void FAudioSink::AddAudioConsumer(IPixelStreaming2AudioConsumer* AudioConsumer)
 	{
+		FScopeLock Lock(&AudioConsumersCS);
 		bool bAlreadyInSet = false;
 		AudioConsumers.Add(AudioConsumer, &bAlreadyInSet);
 		if (!bAlreadyInSet)
@@ -33,6 +35,7 @@ namespace UE::PixelStreaming2
 
 	void FAudioSink::RemoveAudioConsumer(IPixelStreaming2AudioConsumer* AudioConsumer)
 	{
+		FScopeLock Lock(&AudioConsumersCS);
 		if (AudioConsumers.Contains(AudioConsumer))
 		{
 			AudioConsumers.Remove(AudioConsumer);
@@ -57,6 +60,7 @@ namespace UE::PixelStreaming2
 			return;
 		}
 		// Iterate audio consumers and pass this data to their buffers
+		FScopeLock Lock(&AudioConsumersCS);
 		for (IPixelStreaming2AudioConsumer* AudioConsumer : AudioConsumers)
 		{
 			AudioConsumer->ConsumeRawPCM(AudioData, SampleRate, NumChannels, NumFrames);
