Fix Pixel Streaming 2 setting incorrect EpicRtc Log.

Problem was Unreal log values were directly copied into EpicRtc log levels. But EpicRtc levels are reversed.
Solution was to reverse the map and add a missing level. static_asserts also added to verify results.

diff --git a/Engine/Plugins/Media/PixelStreaming2/Source/PixelStreaming2/Internal/EpicRtcLogging.h b/Engine/Plugins/Media/PixelStreaming2/Source/PixelStreaming2/Internal/EpicRtcLogging.h
--- a/Engine/Plugins/Media/PixelStreaming2/Source/PixelStreaming2/Internal/EpicRtcLogging.h
+++ b/Engine/Plugins/Media/PixelStreaming2/Source/PixelStreaming2/Internal/EpicRtcLogging.h
@@ -7,7 +7,7 @@
 
 #include "epic_rtc/common/logging.h"
 
-DECLARE_LOG_CATEGORY_EXTERN(LogPixelStreaming2EpicRtc, Log, All);
+DECLARE_LOG_CATEGORY_EXTERN(LogPixelStreaming2EpicRtc, NoLogging, All);
 
 namespace UE::PixelStreaming2
 {
diff --git a/Engine/Plugins/Media/PixelStreaming2/Source/PixelStreaming2/Private/PixelStreaming2Module.cpp b/Engine/Plugins/Media/PixelStreaming2/Source/PixelStreaming2/Private/PixelStreaming2Module.cpp
--- a/Engine/Plugins/Media/PixelStreaming2/Source/PixelStreaming2/Private/PixelStreaming2Module.cpp
+++ b/Engine/Plugins/Media/PixelStreaming2/Source/PixelStreaming2/Private/PixelStreaming2Module.cpp
@@ -70,6 +70,28 @@ THIRD_PARTY_INCLUDES_END
 
 namespace UE::PixelStreaming2
 {
+	constexpr EpicRtcLogLevel UnrealLogToEpicRtcCategoryMap[] = {
+		EpicRtcLogLevel::Off,
+		EpicRtcLogLevel::Critical,
+		EpicRtcLogLevel::Error,
+		EpicRtcLogLevel::Warning,
+		EpicRtcLogLevel::Info,
+		EpicRtcLogLevel::Info,
+		EpicRtcLogLevel::Debug,
+		EpicRtcLogLevel::Trace,
+		EpicRtcLogLevel::Trace
+	};
+
+	static_assert(EpicRtcLogLevel::Off == UnrealLogToEpicRtcCategoryMap[ELogVerbosity::Type::NoLogging]);
+	static_assert(EpicRtcLogLevel::Critical == UnrealLogToEpicRtcCategoryMap[ELogVerbosity::Type::Fatal]);
+	static_assert(EpicRtcLogLevel::Error == UnrealLogToEpicRtcCategoryMap[ELogVerbosity::Type::Error]);
+	static_assert(EpicRtcLogLevel::Warning == UnrealLogToEpicRtcCategoryMap[ELogVerbosity::Type::Warning]);
+	static_assert(EpicRtcLogLevel::Info == UnrealLogToEpicRtcCategoryMap[ELogVerbosity::Type::Display]);
+	static_assert(EpicRtcLogLevel::Info == UnrealLogToEpicRtcCategoryMap[ELogVerbosity::Type::Log]);
+	static_assert(EpicRtcLogLevel::Debug == UnrealLogToEpicRtcCategoryMap[ELogVerbosity::Type::Verbose]);
+	static_assert(EpicRtcLogLevel::Trace == UnrealLogToEpicRtcCategoryMap[ELogVerbosity::Type::VeryVerbose]);
+	static_assert(EpicRtcLogLevel::Trace == UnrealLogToEpicRtcCategoryMap[ELogVerbosity::Type::All]);
+	
 	FPixelStreaming2Module* FPixelStreaming2Module::PixelStreaming2Module = nullptr;
 
 	FUtf8String FPixelStreaming2Module::EpicRtcConferenceName("pixel_streaming_conference_instance");
@@ -637,16 +659,6 @@ namespace UE::PixelStreaming2
 			return false;
 		}
 
-		static const EpicRtcLogLevel UnrealLogToEpicRtcCategoryMap[] = {
-			EpicRtcLogLevel::Trace,
-			EpicRtcLogLevel::Debug,
-			EpicRtcLogLevel::Info,
-			EpicRtcLogLevel::Warning,
-			EpicRtcLogLevel::Error,
-			EpicRtcLogLevel::Critical,
-			EpicRtcLogLevel::Off
-		};
-
 		FUtf8String EpicRtcFieldTrials(GetFieldTrials());
 
 		WebsocketFactory = MakeRefCount<FEpicRtcWebsocketFactory>();
diff --git a/Engine/Plugins/Media/PixelStreaming2/Source/PixelStreaming2/Private/Tests/MockPlayer.cpp b/Engine/Plugins/Media/PixelStreaming2/Source/PixelStreaming2/Private/Tests/MockPlayer.cpp
--- a/Engine/Plugins/Media/PixelStreaming2/Source/PixelStreaming2/Private/Tests/MockPlayer.cpp
+++ b/Engine/Plugins/Media/PixelStreaming2/Source/PixelStreaming2/Private/Tests/MockPlayer.cpp
@@ -66,16 +66,6 @@ namespace UE::PixelStreaming2
 
 		TRefCountPtr<FEpicRtcWebsocketFactory> WebsocketFactory = MakeRefCount<FEpicRtcWebsocketFactory>(false);
 
-		static const EpicRtcLogLevel UnrealLogToEpicRtcCategoryMap[] = {
-			EpicRtcLogLevel::Trace,
-			EpicRtcLogLevel::Debug,
-			EpicRtcLogLevel::Info,
-			EpicRtcLogLevel::Warning,
-			EpicRtcLogLevel::Error,
-			EpicRtcLogLevel::Critical,
-			EpicRtcLogLevel::Off
-		};
-
 		Result = Platform->CreateConference(ToEpicRtcStringView(ConferenceId),
 			{ ._websocketFactory = WebsocketFactory.GetReference(),
 				._signallingType = EpicRtcSignallingType::PixelStreaming,
