[PS2] Fix: Rename tests to not conflict with PS1 test names

diff --git a/Engine/Plugins/Media/PixelStreaming2/Source/PixelStreaming2/Private/Tests/DataChannelTest.cpp b/Engine/Plugins/Media/PixelStreaming2/Source/PixelStreaming2/Private/Tests/DataChannelTest.cpp
--- a/Engine/Plugins/Media/PixelStreaming2/Source/PixelStreaming2/Private/Tests/DataChannelTest.cpp
+++ b/Engine/Plugins/Media/PixelStreaming2/Source/PixelStreaming2/Private/Tests/DataChannelTest.cpp
@@ -12,7 +12,7 @@
 
 namespace UE::PixelStreaming2
 {
-	IMPLEMENT_SIMPLE_AUTOMATION_TEST(FDataChannelEchoTest, "System.Plugins.PixelStreaming2.FDataChannelEchoTest", EAutomationTestFlags::EditorContext | EAutomationTestFlags::ClientContext | EAutomationTestFlags::ProductFilter)
+	IMPLEMENT_SIMPLE_AUTOMATION_TEST(FDataChannelEchoTest, "System.Plugins.PixelStreaming2.FPS2DataChannelEchoTest", EAutomationTestFlags::EditorContext | EAutomationTestFlags::ClientContext | EAutomationTestFlags::ProductFilter)
 	bool FDataChannelEchoTest::RunTest(const FString& Parameters)
 	{
 		// need to be able to accept codec to handshake otherwise setting local description fails when generating an answer
diff --git a/Engine/Plugins/Media/PixelStreaming2/Source/PixelStreaming2/Private/Tests/EpicRtcConferenceTest.cpp b/Engine/Plugins/Media/PixelStreaming2/Source/PixelStreaming2/Private/Tests/EpicRtcConferenceTest.cpp
--- a/Engine/Plugins/Media/PixelStreaming2/Source/PixelStreaming2/Private/Tests/EpicRtcConferenceTest.cpp
+++ b/Engine/Plugins/Media/PixelStreaming2/Source/PixelStreaming2/Private/Tests/EpicRtcConferenceTest.cpp
@@ -13,7 +13,7 @@
 namespace UE::PixelStreaming2
 {
 	// Tests the creation and removal of a conference from the platform
-	IMPLEMENT_SIMPLE_AUTOMATION_TEST(FEpicRtcConferenceLifetimeTest, "System.Plugins.PixelStreaming2.FEpicRtcConferenceLifetimeTest", EAutomationTestFlags::EditorContext | EAutomationTestFlags::ClientContext | EAutomationTestFlags::ProductFilter)
+	IMPLEMENT_SIMPLE_AUTOMATION_TEST(FEpicRtcConferenceLifetimeTest, "System.Plugins.PixelStreaming2.FPS2EpicRtcConferenceLifetimeTest", EAutomationTestFlags::EditorContext | EAutomationTestFlags::ClientContext | EAutomationTestFlags::ProductFilter)
 	bool FEpicRtcConferenceLifetimeTest::RunTest(const FString& Parameters)
 	{
 		EpicRtcErrorCode Result;
@@ -83,7 +83,7 @@ namespace UE::PixelStreaming2
 	}
 
 	// Tests the conference creation logic to ensure that two conferences with the same name can't be created
-	IMPLEMENT_SIMPLE_AUTOMATION_TEST(FEpicRtcConferenceCreateTest, "System.Plugins.PixelStreaming2.FEpicRtcConferenceCreateTest", EAutomationTestFlags::EditorContext | EAutomationTestFlags::ClientContext | EAutomationTestFlags::ProductFilter)
+	IMPLEMENT_SIMPLE_AUTOMATION_TEST(FEpicRtcConferenceCreateTest, "System.Plugins.PixelStreaming2.FPS2EpicRtcConferenceCreateTest", EAutomationTestFlags::EditorContext | EAutomationTestFlags::ClientContext | EAutomationTestFlags::ProductFilter)
 	bool FEpicRtcConferenceCreateTest::RunTest(const FString& Parameters)
 	{
 		EpicRtcErrorCode Result;
@@ -135,7 +135,7 @@ namespace UE::PixelStreaming2
 	}
 
 	// Tests the conference retrieval logic to ensure that a conference can be retrieved after is has been created
-	IMPLEMENT_SIMPLE_AUTOMATION_TEST(FEpicRtcConferenceGetTest, "System.Plugins.PixelStreaming2.FEpicRtcConferenceGetTest", EAutomationTestFlags::EditorContext | EAutomationTestFlags::ClientContext | EAutomationTestFlags::ProductFilter)
+	IMPLEMENT_SIMPLE_AUTOMATION_TEST(FEpicRtcConferenceGetTest, "System.Plugins.PixelStreaming2.FPS2EpicRtcConferenceGetTest", EAutomationTestFlags::EditorContext | EAutomationTestFlags::ClientContext | EAutomationTestFlags::ProductFilter)
 	bool FEpicRtcConferenceGetTest::RunTest(const FString& Parameters)
 	{
 		EpicRtcErrorCode Result;
diff --git a/Engine/Plugins/Media/PixelStreaming2/Source/PixelStreaming2/Private/Tests/EpicRtcPlatformTest.cpp b/Engine/Plugins/Media/PixelStreaming2/Source/PixelStreaming2/Private/Tests/EpicRtcPlatformTest.cpp
--- a/Engine/Plugins/Media/PixelStreaming2/Source/PixelStreaming2/Private/Tests/EpicRtcPlatformTest.cpp
+++ b/Engine/Plugins/Media/PixelStreaming2/Source/PixelStreaming2/Private/Tests/EpicRtcPlatformTest.cpp
@@ -11,7 +11,7 @@
 
 namespace UE::PixelStreaming2
 {
-	IMPLEMENT_SIMPLE_AUTOMATION_TEST(FEpicRtcPlatformCreateTest, "System.Plugins.PixelStreaming2.FEpicRtcPlatformCreateTest", EAutomationTestFlags::EditorContext | EAutomationTestFlags::ClientContext | EAutomationTestFlags::ProductFilter)
+	IMPLEMENT_SIMPLE_AUTOMATION_TEST(FEpicRtcPlatformCreateTest, "System.Plugins.PixelStreaming2.FPS2EpicRtcPlatformCreateTest", EAutomationTestFlags::EditorContext | EAutomationTestFlags::ClientContext | EAutomationTestFlags::ProductFilter)
 	bool FEpicRtcPlatformCreateTest::RunTest(const FString& Parameters)
 	{
 		EpicRtcErrorCode Result;
@@ -26,7 +26,7 @@ namespace UE::PixelStreaming2
 		return true;
 	}
 
-	IMPLEMENT_SIMPLE_AUTOMATION_TEST(FEpicRtcPlatformLifetimeTest, "System.Plugins.PixelStreaming2.FEpicRtcPlatformLifetimeTest", EAutomationTestFlags::EditorContext | EAutomationTestFlags::ClientContext | EAutomationTestFlags::ProductFilter)
+	IMPLEMENT_SIMPLE_AUTOMATION_TEST(FEpicRtcPlatformLifetimeTest, "System.Plugins.PixelStreaming2.FPS2EpicRtcPlatformLifetimeTest", EAutomationTestFlags::EditorContext | EAutomationTestFlags::ClientContext | EAutomationTestFlags::ProductFilter)
 	bool FEpicRtcPlatformLifetimeTest::RunTest(const FString& Parameters)
 	{
 		EpicRtcPlatformInterface* PlatformPtr;
@@ -56,7 +56,7 @@ namespace UE::PixelStreaming2
 		return true;
 	}
 
-	IMPLEMENT_SIMPLE_AUTOMATION_TEST(FEpicRtcPlatformGetTest, "System.Plugins.PixelStreaming2.FEpicRtcPlatformGetTest", EAutomationTestFlags::EditorContext | EAutomationTestFlags::ClientContext | EAutomationTestFlags::ProductFilter)
+	IMPLEMENT_SIMPLE_AUTOMATION_TEST(FEpicRtcPlatformGetTest, "System.Plugins.PixelStreaming2.FPS2EpicRtcPlatformGetTest", EAutomationTestFlags::EditorContext | EAutomationTestFlags::ClientContext | EAutomationTestFlags::ProductFilter)
 	bool FEpicRtcPlatformGetTest::RunTest(const FString& Parameters)
 	{
 		EpicRtcErrorCode Result;
diff --git a/Engine/Plugins/Media/PixelStreaming2/Source/PixelStreaming2/Private/Tests/EpicRtcRoomTest.cpp b/Engine/Plugins/Media/PixelStreaming2/Source/PixelStreaming2/Private/Tests/EpicRtcRoomTest.cpp
--- a/Engine/Plugins/Media/PixelStreaming2/Source/PixelStreaming2/Private/Tests/EpicRtcRoomTest.cpp
+++ b/Engine/Plugins/Media/PixelStreaming2/Source/PixelStreaming2/Private/Tests/EpicRtcRoomTest.cpp
@@ -12,7 +12,7 @@
 namespace UE::PixelStreaming2
 {
 	// Tests the creation and removal of a session
-	IMPLEMENT_SIMPLE_AUTOMATION_TEST(FEpicRtcRoomLifetimeTest, "System.Plugins.PixelStreaming2.FEpicRtcRoomLifetimeTest", EAutomationTestFlags::EditorContext | EAutomationTestFlags::ClientContext | EAutomationTestFlags::ProductFilter)
+	IMPLEMENT_SIMPLE_AUTOMATION_TEST(FEpicRtcRoomLifetimeTest, "System.Plugins.PixelStreaming2.FPS2EpicRtcRoomLifetimeTest", EAutomationTestFlags::EditorContext | EAutomationTestFlags::ClientContext | EAutomationTestFlags::ProductFilter)
 	bool FEpicRtcRoomLifetimeTest::RunTest(const FString& Parameters)
 	{
 		EpicRtcErrorCode Result;
diff --git a/Engine/Plugins/Media/PixelStreaming2/Source/PixelStreaming2/Private/Tests/EpicRtcSessionTest.cpp b/Engine/Plugins/Media/PixelStreaming2/Source/PixelStreaming2/Private/Tests/EpicRtcSessionTest.cpp
--- a/Engine/Plugins/Media/PixelStreaming2/Source/PixelStreaming2/Private/Tests/EpicRtcSessionTest.cpp
+++ b/Engine/Plugins/Media/PixelStreaming2/Source/PixelStreaming2/Private/Tests/EpicRtcSessionTest.cpp
@@ -14,7 +14,7 @@
 namespace UE::PixelStreaming2
 {
 	// Tests the creation and removal of a session
-	IMPLEMENT_SIMPLE_AUTOMATION_TEST(FEpicRtcSessionLifetimeTest, "System.Plugins.PixelStreaming2.FEpicRtcSessionLifetimeTest", EAutomationTestFlags::EditorContext | EAutomationTestFlags::ClientContext | EAutomationTestFlags::ProductFilter)
+	IMPLEMENT_SIMPLE_AUTOMATION_TEST(FEpicRtcSessionLifetimeTest, "System.Plugins.PixelStreaming2.FPS2EpicRtcSessionLifetimeTest", EAutomationTestFlags::EditorContext | EAutomationTestFlags::ClientContext | EAutomationTestFlags::ProductFilter)
 	bool FEpicRtcSessionLifetimeTest::RunTest(const FString& Parameters)
 	{
 		EpicRtcErrorCode Result;
diff --git a/Engine/Plugins/Media/PixelStreaming2/Source/PixelStreaming2/Private/Tests/ProtocolTest.cpp b/Engine/Plugins/Media/PixelStreaming2/Source/PixelStreaming2/Private/Tests/ProtocolTest.cpp
--- a/Engine/Plugins/Media/PixelStreaming2/Source/PixelStreaming2/Private/Tests/ProtocolTest.cpp
+++ b/Engine/Plugins/Media/PixelStreaming2/Source/PixelStreaming2/Private/Tests/ProtocolTest.cpp
@@ -15,7 +15,7 @@
 
 namespace UE::PixelStreaming2
 {
-	IMPLEMENT_SIMPLE_AUTOMATION_TEST(FProtocolTestAddMessage, "System.Plugins.PixelStreaming2.FProtocolTestAddMessage", EAutomationTestFlags::EditorContext | EAutomationTestFlags::ClientContext | EAutomationTestFlags::ProductFilter)
+	IMPLEMENT_SIMPLE_AUTOMATION_TEST(FProtocolTestAddMessage, "System.Plugins.PixelStreaming2.FPS2ProtocolTestAddMessage", EAutomationTestFlags::EditorContext | EAutomationTestFlags::ClientContext | EAutomationTestFlags::ProductFilter)
 	bool FProtocolTestAddMessage::RunTest(const FString& Parameters)
 	{
 		int32 StreamerPort = TestUtils::NextStreamerPort();
@@ -90,7 +90,7 @@ namespace UE::PixelStreaming2
 		return true;
 	}
 
-	IMPLEMENT_SIMPLE_AUTOMATION_TEST(FProtocolTestUseCustomMessage, "System.Plugins.PixelStreaming2.FProtocolTestUseCustomMessage", EAutomationTestFlags::EditorContext | EAutomationTestFlags::ClientContext | EAutomationTestFlags::ProductFilter)
+	IMPLEMENT_SIMPLE_AUTOMATION_TEST(FProtocolTestUseCustomMessage, "System.Plugins.PixelStreaming2.FPS2ProtocolTestUseCustomMessage", EAutomationTestFlags::EditorContext | EAutomationTestFlags::ClientContext | EAutomationTestFlags::ProductFilter)
 	bool FProtocolTestUseCustomMessage::RunTest(const FString& Parameters)
 	{
 		int32 StreamerPort = TestUtils::NextStreamerPort();
diff --git a/Engine/Plugins/Media/PixelStreaming2/Source/PixelStreaming2/Private/Tests/VP8Tests.cpp b/Engine/Plugins/Media/PixelStreaming2/Source/PixelStreaming2/Private/Tests/VP8Tests.cpp
--- a/Engine/Plugins/Media/PixelStreaming2/Source/PixelStreaming2/Private/Tests/VP8Tests.cpp
+++ b/Engine/Plugins/Media/PixelStreaming2/Source/PixelStreaming2/Private/Tests/VP8Tests.cpp
@@ -7,7 +7,7 @@
 
 namespace UE::PixelStreaming2
 {
-	IMPLEMENT_SIMPLE_AUTOMATION_TEST(FVP8FrameReceivedTest, "System.Plugins.PixelStreaming2.FVP8FrameReceivedTest", EAutomationTestFlags::EditorContext | EAutomationTestFlags::ClientContext | EAutomationTestFlags::ProductFilter)
+	IMPLEMENT_SIMPLE_AUTOMATION_TEST(FVP8FrameReceivedTest, "System.Plugins.PixelStreaming2.FPS2VP8FrameReceivedTest", EAutomationTestFlags::EditorContext | EAutomationTestFlags::ClientContext | EAutomationTestFlags::ProductFilter)
 	bool FVP8FrameReceivedTest::RunTest(const FString& Parameters)
 	{
 		SetCodec(EVideoCodec::VP8);
@@ -15,7 +15,7 @@ namespace UE::PixelStreaming2
 		return true;
 	}
 
-	IMPLEMENT_SIMPLE_AUTOMATION_TEST(FVP8FrameResizeTest, "System.Plugins.PixelStreaming2.FVP8FrameResizeTest", EAutomationTestFlags::EditorContext | EAutomationTestFlags::ClientContext | EAutomationTestFlags::ProductFilter)
+	IMPLEMENT_SIMPLE_AUTOMATION_TEST(FVP8FrameResizeTest, "System.Plugins.PixelStreaming2.FPS2VP8FrameResizeTest", EAutomationTestFlags::EditorContext | EAutomationTestFlags::ClientContext | EAutomationTestFlags::ProductFilter)
 	bool FVP8FrameResizeTest::RunTest(const FString& Parameters)
 	{
 		SetCodec(EVideoCodec::VP8);
diff --git a/Engine/Plugins/Media/PixelStreaming2/Source/PixelStreaming2/Private/Tests/VP9Tests.cpp b/Engine/Plugins/Media/PixelStreaming2/Source/PixelStreaming2/Private/Tests/VP9Tests.cpp
--- a/Engine/Plugins/Media/PixelStreaming2/Source/PixelStreaming2/Private/Tests/VP9Tests.cpp
+++ b/Engine/Plugins/Media/PixelStreaming2/Source/PixelStreaming2/Private/Tests/VP9Tests.cpp
@@ -7,7 +7,7 @@
 
 namespace UE::PixelStreaming2
 {
-	IMPLEMENT_SIMPLE_AUTOMATION_TEST(FVP9FrameReceivedTest, "System.Plugins.PixelStreaming2.FVP9FrameReceivedTest", EAutomationTestFlags::EditorContext | EAutomationTestFlags::ClientContext | EAutomationTestFlags::ProductFilter)
+	IMPLEMENT_SIMPLE_AUTOMATION_TEST(FVP9FrameReceivedTest, "System.Plugins.PixelStreaming2.FPS2VP9FrameReceivedTest", EAutomationTestFlags::EditorContext | EAutomationTestFlags::ClientContext | EAutomationTestFlags::ProductFilter)
 	bool FVP9FrameReceivedTest::RunTest(const FString& Parameters)
 	{
 		SetCodec(EVideoCodec::VP9);
@@ -15,7 +15,7 @@ namespace UE::PixelStreaming2
 		return true;
 	}
 
-	IMPLEMENT_SIMPLE_AUTOMATION_TEST(FVP9FrameResizeTest, "System.Plugins.PixelStreaming2.FVP9FrameResizeTest", EAutomationTestFlags::EditorContext | EAutomationTestFlags::ClientContext | EAutomationTestFlags::ProductFilter)
+	IMPLEMENT_SIMPLE_AUTOMATION_TEST(FVP9FrameResizeTest, "System.Plugins.PixelStreaming2.FPS2VP9FrameResizeTest", EAutomationTestFlags::EditorContext | EAutomationTestFlags::ClientContext | EAutomationTestFlags::ProductFilter)
 	bool FVP9FrameResizeTest::RunTest(const FString& Parameters)
 	{
 		SetCodec(EVideoCodec::VP9);
