[PS2] QOL: Move blueprint classes to a Blueprints private folder

diff --git a/Engine/Plugins/Media/PixelStreaming2/Source/PixelStreaming2/Internal/PixelStreaming2MediaTexture.h b/Engine/Plugins/Media/PixelStreaming2/Source/PixelStreaming2/Internal/PixelStreaming2MediaTexture.h
--- a/Engine/Plugins/Media/PixelStreaming2/Source/PixelStreaming2/Internal/PixelStreaming2MediaTexture.h
+++ b/Engine/Plugins/Media/PixelStreaming2/Source/PixelStreaming2/Internal/PixelStreaming2MediaTexture.h
@@ -8,7 +8,10 @@
 
 #include "PixelStreaming2MediaTexture.generated.h"
 
-class FPixelStreaming2MediaTextureResource;
+namespace UE::PixelStreaming2
+{
+	class FPixelStreaming2MediaTextureResource;
+} // namespace UE::PixelStreaming2
 
 /**
  * A Texture Object that can be used in materials etc. that takes updates from webrtc frames.
@@ -39,5 +42,5 @@ private:
 
 	FCriticalSection RenderSyncContext;
 	// NOTE: This type has to be a raw ptr because UE will manage the lifetime of it
-	FPixelStreaming2MediaTextureResource* CurrentResource;
+	UE::PixelStreaming2::FPixelStreaming2MediaTextureResource* CurrentResource;
 };
diff --git a/Engine/Plugins/Media/PixelStreaming2/Source/PixelStreaming2/Private/Blueprints/PixelStreaming2AudioComponent.cpp b/Engine/Plugins/Media/PixelStreaming2/Source/PixelStreaming2/Private/Blueprints/PixelStreaming2AudioComponent.cpp
new file mode 100644
--- /dev/null
+++ b/Engine/Plugins/Media/PixelStreaming2/Source/PixelStreaming2/Private/Blueprints/PixelStreaming2AudioComponent.cpp
@@ -0,0 +1,201 @@
+// Copyright Epic Games, Inc. All Rights Reserved.
+
+#include "PixelStreaming2AudioComponent.h"
+
+#include "CoreMinimal.h"
+#include "IPixelStreaming2Module.h"
+#include "IPixelStreaming2Streamer.h"
+#include "Logging.h"
+#include "SampleBuffer.h"
+#include "SoundGenerator.h"
+
+/*
+ * Component that recieves audio from a remote webrtc connection and outputs it into UE using a "synth component".
+ */
+UPixelStreaming2AudioComponent::UPixelStreaming2AudioComponent(const FObjectInitializer& ObjectInitializer)
+	: Super(ObjectInitializer)
+	, PlayerToHear(FString())
+	, bAutoFindPeer(true)
+	, AudioSink(nullptr)
+	, SoundGenerator(MakeShared<UE::PixelStreaming2::FSoundGenerator, ESPMode::ThreadSafe>())
+{
+	PreferredBufferLength = 512u;
+	NumChannels = 2;
+	PrimaryComponentTick.bCanEverTick = true;
+	SetComponentTickEnabled(true);
+	bAutoActivate = true;
+};
+
+ISoundGeneratorPtr UPixelStreaming2AudioComponent::CreateSoundGenerator(const FSoundGeneratorInitParams& InParams)
+{
+	SoundGenerator->SetParameters(InParams);
+	Initialize(InParams.SampleRate);
+	return SoundGenerator;
+}
+
+void UPixelStreaming2AudioComponent::OnBeginGenerate()
+{
+	SoundGenerator->bGeneratingAudio = true;
+}
+
+void UPixelStreaming2AudioComponent::OnEndGenerate()
+{
+	SoundGenerator->bGeneratingAudio = false;
+}
+
+void UPixelStreaming2AudioComponent::BeginDestroy()
+{
+	Super::BeginDestroy();
+	Reset();
+}
+
+bool UPixelStreaming2AudioComponent::ListenTo(FString PlayerToListenTo)
+{
+	IPixelStreaming2Module& PixelStreaming2Module = IPixelStreaming2Module::Get();
+	if (!PixelStreaming2Module.IsReady())
+	{
+		return false;
+	}
+	return StreamerListenTo(PixelStreaming2Module.GetDefaultStreamerID(), PlayerToListenTo);
+}
+
+bool UPixelStreaming2AudioComponent::StreamerListenTo(FString StreamerId, FString PlayerToListenTo)
+{
+	if (!IPixelStreaming2Module::IsAvailable())
+	{
+		UE_LOG(LogPixelStreaming2, Verbose, TEXT("Pixel Streaming audio component could not listen to anything because Pixel Streaming module is not loaded. This is expected on dedicated servers."));
+		return false;
+	}
+
+	IPixelStreaming2Module& PixelStreaming2Module = IPixelStreaming2Module::Get();
+	if (!PixelStreaming2Module.IsReady())
+	{
+		return false;
+	}
+
+	PlayerToHear = PlayerToListenTo;
+
+	if (StreamerId == FString())
+	{
+		TArray<FString> StreamerIds = PixelStreaming2Module.GetStreamerIds();
+		if (StreamerIds.Num() > 0)
+		{
+			StreamerToHear = StreamerIds[0];
+		}
+		else
+		{
+			StreamerToHear = PixelStreaming2Module.GetDefaultStreamerID();
+		}
+	}
+	else
+	{
+		StreamerToHear = StreamerId;
+	}
+
+	TSharedPtr<IPixelStreaming2Streamer> Streamer = PixelStreaming2Module.FindStreamer(StreamerToHear);
+	if (!Streamer)
+	{
+		return false;
+	}
+	IPixelStreaming2AudioSink* CandidateSink = WillListenToAnyPlayer() ? Streamer->GetUnlistenedAudioSink() : Streamer->GetPeerAudioSink(FString(PlayerToHear));
+
+	if (CandidateSink == nullptr)
+	{
+		return false;
+	}
+
+	AudioSink = CandidateSink;
+	AudioSink->AddAudioConsumer(this);
+
+	return true;
+}
+
+void UPixelStreaming2AudioComponent::Reset()
+{
+	PlayerToHear = FString();
+	StreamerToHear = FString();
+	SoundGenerator->bShouldGenerateAudio = false;
+	if (AudioSink)
+	{
+		AudioSink->RemoveAudioConsumer(this);
+	}
+	AudioSink = nullptr;
+	SoundGenerator->EmptyBuffers();
+}
+
+bool UPixelStreaming2AudioComponent::IsListeningToPlayer()
+{
+	return SoundGenerator->bShouldGenerateAudio;
+}
+
+bool UPixelStreaming2AudioComponent::WillListenToAnyPlayer()
+{
+	return PlayerToHear == FString();
+}
+
+void UPixelStreaming2AudioComponent::ConsumeRawPCM(const int16_t* AudioData, int InSampleRate, size_t NChannels, size_t NFrames)
+{
+	// Sound generator has not been initialized yet.
+	if (SoundGenerator->GetSampleRate() == 0 || GetAudioComponent() == nullptr)
+	{
+		return;
+	}
+
+	// Set pitch multiplier as a way to handle mismatched sample rates
+	if (InSampleRate != SoundGenerator->GetSampleRate())
+	{
+		GetAudioComponent()->SetPitchMultiplier((float)InSampleRate / SoundGenerator->GetSampleRate());
+	}
+	else if (GetAudioComponent()->PitchMultiplier != 1.0f)
+	{
+		GetAudioComponent()->SetPitchMultiplier(1.0f);
+	}
+
+	Audio::TSampleBuffer<int16_t> Buffer(AudioData, NFrames, NChannels, InSampleRate);
+	if (NChannels != SoundGenerator->GetNumChannels())
+	{
+		Buffer.MixBufferToChannels(SoundGenerator->GetNumChannels());
+	}
+
+	SoundGenerator->AddAudio(Buffer.GetData(), InSampleRate, NChannels, Buffer.GetNumSamples());
+}
+
+void UPixelStreaming2AudioComponent::OnConsumerAdded()
+{
+	SoundGenerator->bShouldGenerateAudio = true;
+	Start();
+}
+
+void UPixelStreaming2AudioComponent::OnConsumerRemoved()
+{
+	Reset();
+}
+
+void UPixelStreaming2AudioComponent::TickComponent(float DeltaTime, ELevelTick TickType, FActorComponentTickFunction* ThisTickFunction)
+{
+
+	bool bPixelStreaming2Loaded = IPixelStreaming2Module::IsAvailable();
+
+	// Early out if running in commandlet
+	if (IsRunningCommandlet())
+	{
+		return;
+	}
+
+	// if auto connect turned off don't bother
+	if (!bAutoFindPeer)
+	{
+		return;
+	}
+
+	// if listening to a peer don't auto connect
+	if (IsListeningToPlayer())
+	{
+		return;
+	}
+
+	if (StreamerListenTo(StreamerToHear, PlayerToHear))
+	{
+		UE_LOG(LogPixelStreaming2, Log, TEXT("PixelStreaming2 audio component found a WebRTC peer to listen to."));
+	}
+}
diff --git a/Engine/Plugins/Media/PixelStreaming2/Source/PixelStreaming2/Private/Blueprints/PixelStreaming2AudioComponent.h b/Engine/Plugins/Media/PixelStreaming2/Source/PixelStreaming2/Private/Blueprints/PixelStreaming2AudioComponent.h
new file mode 100644
--- /dev/null
+++ b/Engine/Plugins/Media/PixelStreaming2/Source/PixelStreaming2/Private/Blueprints/PixelStreaming2AudioComponent.h
@@ -0,0 +1,92 @@
+// Copyright Epic Games, Inc. All Rights Reserved.
+
+#pragma once
+
+#include "CoreMinimal.h"
+#include "Components/SynthComponent.h"
+#include "IPixelStreaming2AudioConsumer.h"
+#include "IPixelStreaming2AudioSink.h"
+#include "PixelStreaming2AudioComponent.generated.h"
+
+namespace UE::PixelStreaming2
+{
+	class FSoundGenerator;
+} // namespace UE::PixelStreaming2
+
+/**
+ * Allows in-engine playback of incoming WebRTC audio from a particular Pixel Streaming player/peer using their mic in the browser.
+ * Note: Each audio component associates itself with a particular Pixel Streaming player/peer (using the the Pixel Streaming player id).
+ */
+UCLASS(Blueprintable, ClassGroup = (PixelStreaming2), meta = (BlueprintSpawnableComponent))
+class PIXELSTREAMING2_API UPixelStreaming2AudioComponent : public USynthComponent, public IPixelStreaming2AudioConsumer
+{
+	GENERATED_BODY()
+
+protected:
+	UPixelStreaming2AudioComponent(const FObjectInitializer& ObjectInitializer);
+
+	//~ Begin USynthComponent interface
+	virtual void			   OnBeginGenerate() override;
+	virtual void			   OnEndGenerate() override;
+	virtual ISoundGeneratorPtr CreateSoundGenerator(const FSoundGeneratorInitParams& InParams) override;
+	//~ End USynthComponent interface
+
+	//~ Begin UObject interface
+	virtual void BeginDestroy() override;
+	//~ End UObject interface
+
+	//~ Begin UActorComponent interface
+	virtual void TickComponent(float DeltaTime, ELevelTick TickType, FActorComponentTickFunction* ThisTickFunction) override;
+	//~ End UActorComponent interface
+
+public:
+	/**
+	 *   The Pixel Streaming streamer of the player that we wish to listen to.
+	 *   If this is left blank this component will use the default Pixel Streaming streamer
+	 */
+	UPROPERTY(EditAnywhere, BlueprintReadWrite, Category = "Pixel Streaming Audio Component")
+	FString StreamerToHear;
+
+	/**
+	 *   The Pixel Streaming player/peer whose audio we wish to listen to.
+	 *   If this is left blank this component will listen to the first non-listened to peer that connects after this component is ready.
+	 *   Note: that when the listened to peer disconnects this component is reset to blank and will once again listen to the next non-listened to peer that connects.
+	 */
+	UPROPERTY(EditAnywhere, BlueprintReadWrite, Category = "Pixel Streaming Audio Component")
+	FString PlayerToHear;
+
+	/**
+	 *  If not already listening to a player/peer will try to attach for listening to the "PlayerToHear" each tick.
+	 */
+	UPROPERTY(EditAnywhere, BlueprintReadWrite, Category = "Pixel Streaming Audio Component")
+	bool bAutoFindPeer;
+
+private:
+	IPixelStreaming2AudioSink*											AudioSink;
+	TSharedPtr<UE::PixelStreaming2::FSoundGenerator, ESPMode::ThreadSafe> SoundGenerator;
+
+public:
+	// Listen to a specific player on the default streamer. If the player is not found this component will be silent.
+	UFUNCTION(BlueprintCallable, Category = "Pixel Streaming Audio Component")
+	bool ListenTo(FString PlayerToListenTo);
+
+	// Listen to a specific player. If the player is not found this component will be silent.
+	UFUNCTION(BlueprintCallable, Category = "Pixel Streaming Audio Component")
+	bool StreamerListenTo(FString StreamerId, FString PlayerToListenTo);
+
+	// True if listening to a connected WebRTC peer through Pixel Streaming.
+	UFUNCTION(BlueprintCallable, Category = "Pixel Streaming Audio Component")
+	bool IsListeningToPlayer();
+
+	bool WillListenToAnyPlayer();
+
+	// Stops listening to any connected player/peer and resets internal state so component is ready to listen again.
+	UFUNCTION(BlueprintCallable, Category = "Pixel Streaming Audio Component")
+	void Reset();
+
+	//~ Begin IPixelStreaming2AudioConsumer interface
+	void ConsumeRawPCM(const int16_t* AudioData, int InSampleRate, size_t NChannels, size_t NFrames);
+	void OnConsumerAdded();
+	void OnConsumerRemoved();
+	//~ End IPixelStreaming2AudioConsumer interface
+};
\ No newline at end of file
diff --git a/Engine/Plugins/Media/PixelStreaming2/Source/PixelStreaming2/Private/Blueprints/PixelStreaming2Blueprints.cpp b/Engine/Plugins/Media/PixelStreaming2/Source/PixelStreaming2/Private/Blueprints/PixelStreaming2Blueprints.cpp
new file mode 100644
--- /dev/null
+++ b/Engine/Plugins/Media/PixelStreaming2/Source/PixelStreaming2/Private/Blueprints/PixelStreaming2Blueprints.cpp
@@ -0,0 +1,151 @@
+// Copyright Epic Games, Inc. All Rights Reserved.
+
+#include "PixelStreaming2Blueprints.h"
+
+#include "IPixelStreaming2Streamer.h"
+#include "Logging.h"
+#include "Misc/FileHelper.h"
+#include "PixelStreaming2AudioComponent.h"
+#include "PixelStreaming2InputComponent.h"
+
+TSharedPtr<IPixelStreaming2Streamer> GetStreamer(FString StreamerId)
+{
+	IPixelStreaming2Module* Module = UE::PixelStreaming2::FPixelStreaming2Module::GetModule();
+	if (!Module)
+	{
+		UE_LOG(LogPixelStreaming2, Error, TEXT("Could not get Pixel Streaming module, so unable to get streamer by id."));
+		return nullptr;
+	}
+	TSharedPtr<IPixelStreaming2Streamer> Streamer = Module->FindStreamer(StreamerId);
+	if (Streamer)
+	{
+		return Streamer;
+	}
+	else
+	{
+		UE_LOG(LogPixelStreaming2, Error, TEXT("Unable to get streamer with id %s"), *StreamerId);
+		return nullptr;
+	}
+}
+
+void UPixelStreaming2Blueprints::SendFileAsByteArray(TArray<uint8> ByteArray, FString MimeType, FString FileExtension)
+{
+	StreamerSendFileAsByteArray(GetDefaultStreamerID(), ByteArray, MimeType, FileExtension);
+}
+
+void UPixelStreaming2Blueprints::StreamerSendFileAsByteArray(FString StreamerId, TArray<uint8> ByteArray, FString MimeType, FString FileExtension)
+{
+	if (TSharedPtr<IPixelStreaming2Streamer> Streamer = GetStreamer(StreamerId))
+	{
+		const TArray64<uint8> LargeByteArray(ByteArray);
+		Streamer->SendFileData(LargeByteArray, MimeType, FileExtension);
+	}
+}
+
+void UPixelStreaming2Blueprints::SendFile(FString FilePath, FString MimeType, FString FileExtension)
+{
+	StreamerSendFile(GetDefaultStreamerID(), FilePath, MimeType, FileExtension);
+}
+
+void UPixelStreaming2Blueprints::StreamerSendFile(FString StreamerId, FString FilePath, FString MimeType, FString FileExtension)
+{
+	if (TSharedPtr<IPixelStreaming2Streamer> Streamer = GetStreamer(StreamerId))
+	{
+		TArray64<uint8> ByteData;
+		bool			bSuccess = FFileHelper::LoadFileToArray(ByteData, *FilePath);
+		if (bSuccess)
+		{
+			Streamer->SendFileData(ByteData, MimeType, FileExtension);
+		}
+		else
+		{
+			UE_LOG(LogPixelStreaming2, Error, TEXT("FileHelper failed to load file data"));
+		}
+	}
+}
+
+void UPixelStreaming2Blueprints::ForceKeyFrame()
+{
+	StreamerForceKeyFrame(GetDefaultStreamerID());
+}
+
+void UPixelStreaming2Blueprints::StreamerForceKeyFrame(FString StreamerId)
+{
+	if (TSharedPtr<IPixelStreaming2Streamer> Streamer = GetStreamer(StreamerId))
+	{
+		Streamer->ForceKeyFrame();
+	}
+}
+
+void UPixelStreaming2Blueprints::FreezeFrame(UTexture2D* Texture)
+{
+	StreamerFreezeStream(GetDefaultStreamerID(), Texture);
+}
+
+void UPixelStreaming2Blueprints::StreamerFreezeStream(FString StreamerId, UTexture2D* Texture)
+{
+	if (TSharedPtr<IPixelStreaming2Streamer> Streamer = GetStreamer(StreamerId))
+	{
+		Streamer->FreezeStream(Texture);
+	}
+}
+
+void UPixelStreaming2Blueprints::UnfreezeFrame()
+{
+	StreamerUnfreezeStream(GetDefaultStreamerID());
+}
+
+void UPixelStreaming2Blueprints::StreamerUnfreezeStream(FString StreamerId)
+{
+	if (TSharedPtr<IPixelStreaming2Streamer> Streamer = GetStreamer(StreamerId))
+	{
+		Streamer->UnfreezeStream();
+	}
+}
+
+void UPixelStreaming2Blueprints::KickPlayer(FString PlayerId)
+{
+	StreamerKickPlayer(GetDefaultStreamerID(), PlayerId);
+}
+
+void UPixelStreaming2Blueprints::StreamerKickPlayer(FString StreamerId, FString PlayerId)
+{
+	if (TSharedPtr<IPixelStreaming2Streamer> Streamer = GetStreamer(StreamerId))
+	{
+		Streamer->KickPlayer(PlayerId);
+	}
+}
+
+TArray<FString> UPixelStreaming2Blueprints::GetConnectedPlayers()
+{
+	return StreamerGetConnectedPlayers(GetDefaultStreamerID());
+}
+
+TArray<FString> UPixelStreaming2Blueprints::StreamerGetConnectedPlayers(FString StreamerId)
+{
+	if (TSharedPtr<IPixelStreaming2Streamer> Streamer = GetStreamer(StreamerId))
+	{
+		return Streamer->GetConnectedPlayers();
+	}
+	else
+	{
+		UE_LOG(LogPixelStreaming2, Error, TEXT("No streamer found with specified id - %s. Returning no connected players."), *StreamerId);
+		return {};
+	}
+}
+
+FString UPixelStreaming2Blueprints::GetDefaultStreamerID()
+{
+	IPixelStreaming2Module* Module = UE::PixelStreaming2::FPixelStreaming2Module::GetModule();
+	if (!Module)
+	{
+		UE_LOG(LogPixelStreaming2, Error, TEXT("Pixel Streaming module not initialized yet - cannot get default streamer id."));
+		return FString(TEXT("ModuleNotInitializedYet"));
+	}
+	return Module->GetDefaultStreamerID();
+}
+
+UPixelStreaming2Delegates* UPixelStreaming2Blueprints::GetDelegates()
+{
+	return UPixelStreaming2Delegates::Get();
+}
\ No newline at end of file
diff --git a/Engine/Plugins/Media/PixelStreaming2/Source/PixelStreaming2/Private/Blueprints/PixelStreaming2Blueprints.h b/Engine/Plugins/Media/PixelStreaming2/Source/PixelStreaming2/Private/Blueprints/PixelStreaming2Blueprints.h
new file mode 100644
--- /dev/null
+++ b/Engine/Plugins/Media/PixelStreaming2/Source/PixelStreaming2/Private/Blueprints/PixelStreaming2Blueprints.h
@@ -0,0 +1,145 @@
+// Copyright Epic Games, Inc. All Rights Reserved.
+
+#pragma once
+
+#include "CoreMinimal.h"
+#include "IPixelStreaming2Module.h"
+#include "PixelStreaming2Module.h"
+#include "Kismet/BlueprintFunctionLibrary.h"
+#include "Containers/Array.h"
+#include "PixelStreaming2Delegates.h"
+#include "PixelStreaming2Blueprints.generated.h"
+
+UCLASS()
+class PIXELSTREAMING2_API UPixelStreaming2Blueprints : public UBlueprintFunctionLibrary
+{
+	GENERATED_BODY()
+public:
+	/**
+	 * Send a specified byte array over the WebRTC peer connection data channel. The extension and mime type are used for file reconstruction on the front end
+	 *
+	 * @param   ByteArray       The raw data that will be sent over the data channel
+	 * @param   MimeType        The mime type of the file. Used for reconstruction on the front end
+	 * @param   FileExtension   The file extension. Used for file reconstruction on the front end
+	 */
+	UFUNCTION(BlueprintCallable, Category = "Pixel Streaming 2")
+	static void SendFileAsByteArray(TArray<uint8> ByteArray, FString MimeType, FString FileExtension);
+
+	/**
+	 * Send a specified byte array over the WebRTC peer connection data channel. The extension and mime type are used for file reconstruction on the front end
+	 *
+	 * @param	StreamerId		The streamer use when sending the data
+	 * @param   ByteArray       The raw data that will be sent over the data channel
+	 * @param   MimeType        The mime type of the file. Used for reconstruction on the front end
+	 * @param   FileExtension   The file extension. Used for file reconstruction on the front end
+	 */
+	UFUNCTION(BlueprintCallable, Category = "Pixel Streaming 2")
+	static void StreamerSendFileAsByteArray(FString StreamerId, TArray<uint8> ByteArray, FString MimeType, FString FileExtension);
+
+	/**
+	 * Send a specified file over the WebRTC peer connection data channel. The extension and mime type are used for file reconstruction on the front end
+	 *
+	 * @param   FilePath        The path to the file that will be sent
+	 * @param   MimeType        The mime type of the file. Used for file reconstruction on the front end
+	 * @param   FileExtension   The file extension. Used for file reconstruction on the front end
+	 */
+	UFUNCTION(BlueprintCallable, Category = "Pixel Streaming 2")
+	static void SendFile(FString Filepath, FString MimeType, FString FileExtension);
+
+	/**
+	 * Send a specified file over the WebRTC peer connection data channel. The extension and mime type are used for file reconstruction on the front end
+	 *
+	 * @param	StreamerId		The streamer use when sending the data
+	 * @param   FilePath        The path to the file that will be sent
+	 * @param   MimeType        The mime type of the file. Used for file reconstruction on the front end
+	 * @param   FileExtension   The file extension. Used for file reconstruction on the front end
+	 */
+	UFUNCTION(BlueprintCallable, Category = "Pixel Streaming 2")
+	static void StreamerSendFile(FString StreamerId, FString Filepath, FString MimeType, FString FileExtension);
+
+	/**
+	 * Force a key frame to be sent to the default streamer (if there is one).
+	 */
+	UFUNCTION(BlueprintCallable, Category = "Pixel Streaming 2")
+	static void ForceKeyFrame();
+
+	/**
+	 * Force a key frame to be sent to the specified streamer.
+	 */
+	UFUNCTION(BlueprintCallable, Category = "Pixel Streaming 2")
+	static void StreamerForceKeyFrame(FString StreamerId);
+
+	/**
+	 * Freeze the video stream of the default streamer (if there is one).
+	 * @param   Texture         The freeze frame to display. If null then show the last frame that was streamed.
+	 */
+	UFUNCTION(BlueprintCallable, Category = "Pixel Streaming 2")
+	static void FreezeFrame(UTexture2D* Texture);
+
+	/**
+	 * Unfreeze the video stream of the default streamer (if there is one).
+	 */
+	UFUNCTION(BlueprintCallable, Category = "Pixel Streaming 2")
+	static void UnfreezeFrame();
+
+	/**
+	 * Freeze the video stream of the specified video stream.
+	 * @param	StreamerId		The id of the streamer to freeze.
+	 * @param   Texture         The freeze frame to display. If null then show the last frame that was streamed.
+	 */
+	UFUNCTION(BlueprintCallable, Category = "Pixel Streaming 2")
+	static void StreamerFreezeStream(FString StreamerId, UTexture2D* Texture);
+
+	/**
+	 * Unfreeze the video stream of the specified streamer.
+	 * @param StreamerId		The id of the streamer to unfreeze.
+	 */
+	UFUNCTION(BlueprintCallable, Category = "Pixel Streaming 2")
+	static void StreamerUnfreezeStream(FString StreamerId);
+
+	/**
+	 * Kick a player that is connected to the default streamer (if there is a default streamer).
+	 * @param   PlayerId         The ID of the player to kick.
+	 */
+	UFUNCTION(BlueprintCallable, Category = "Pixel Streaming 2")
+	static void KickPlayer(FString PlayerId);
+
+	/**
+	 * Kick a player connected to the specified streamer.
+	 * @param	StreamerId		The streamer which the player belongs
+	 * @param   PlayerId        The ID of the player to kick.
+	 */
+	UFUNCTION(BlueprintCallable, Category = "Pixel Streaming 2")
+	static void StreamerKickPlayer(FString StreamerId, FString PlayerId);
+
+	/**
+	 * @brief Get the connected players
+	 *
+	 * @return TArray<FString> The connected players
+	 */
+	UFUNCTION(BlueprintCallable, Category = "Pixel Streaming 2")
+	static TArray<FString> GetConnectedPlayers();
+
+	/**
+	 * @brief Get the connected players
+	 *
+	 * @param StreamerId	The streamer whose list of players you wish to get
+	 * @return TArray<FString> The connected players
+	 */
+	UFUNCTION(BlueprintCallable, Category = "Pixel Streaming 2")
+	static TArray<FString> StreamerGetConnectedPlayers(FString StreamerId);
+
+	/**
+	 * Get the default Streamer ID
+	 */
+	UFUNCTION(BlueprintCallable, Category = "Pixel Streaming 2")
+	static FString GetDefaultStreamerID();
+
+	// PixelStreaming2Delegates
+	/**
+	 * Get the singleton. This allows application-specific blueprints to bind
+	 * to delegates of interest.
+	 */
+	UFUNCTION(BlueprintCallable, Category = "Pixel Streaming 2")
+	static UPixelStreaming2Delegates* GetDelegates();
+};
\ No newline at end of file
diff --git a/Engine/Plugins/Media/PixelStreaming2/Source/PixelStreaming2/Private/Blueprints/PixelStreaming2InputComponent.cpp b/Engine/Plugins/Media/PixelStreaming2/Source/PixelStreaming2/Private/Blueprints/PixelStreaming2InputComponent.cpp
new file mode 100644
--- /dev/null
+++ b/Engine/Plugins/Media/PixelStreaming2/Source/PixelStreaming2/Private/Blueprints/PixelStreaming2InputComponent.cpp
@@ -0,0 +1,82 @@
+// Copyright Epic Games, Inc. All Rights Reserved.
+
+#include "PixelStreaming2InputComponent.h"
+#include "IPixelStreaming2Module.h"
+#include "PixelStreaming2Module.h"
+#include "Engine/Engine.h"
+#include "Engine/World.h"
+#include "Policies/CondensedJsonPrintPolicy.h"
+#include "Serialization/JsonReader.h"
+#include "Serialization/JsonSerializer.h"
+#include "GameFramework/GameUserSettings.h"
+#include "Logging.h"
+#include "IPixelStreaming2Streamer.h"
+#include "PixelStreaming2Utils.h"
+
+UPixelStreaming2Input::UPixelStreaming2Input(const FObjectInitializer& ObjectInitializer)
+	: Super(ObjectInitializer)
+	, PixelStreaming2Module((UE::PixelStreaming2::FPixelStreaming2Module*)UE::PixelStreaming2::FPixelStreaming2Module::GetModule())
+{
+	bAutoActivate = true;
+	PrimaryComponentTick.bCanEverTick = true;
+	SetComponentTickEnabled(true);
+}
+
+void UPixelStreaming2Input::BeginPlay()
+{
+	Super::BeginPlay();
+
+	if (PixelStreaming2Module)
+	{
+		// When this component is initializing it registers itself with the Pixel Streaming module.
+		PixelStreaming2Module->AddInputComponent(this);
+	}
+	else
+	{
+		UE_LOG(LogPixelStreaming2, Warning, TEXT("Pixel Streaming input component not added because Pixel Streaming module is not loaded. This is expected on dedicated servers."));
+	}
+}
+
+void UPixelStreaming2Input::EndPlay(const EEndPlayReason::Type EndPlayReason)
+{
+	Super::EndPlay(EndPlayReason);
+
+	if (PixelStreaming2Module)
+	{
+		// When this component is destructing it unregisters itself with the Pixel Streaming module.
+		PixelStreaming2Module->RemoveInputComponent(this);
+	}
+	else
+	{
+		UE_LOG(LogPixelStreaming2, Warning, TEXT("Pixel Streaming input component not removed because Pixel Streaming module is not loaded. This is expected on dedicated servers."));
+	}
+}
+
+void UPixelStreaming2Input::SendPixelStreaming2Response(const FString& Descriptor)
+{
+	if (PixelStreaming2Module)
+	{
+		PixelStreaming2Module->ForEachStreamer([&Descriptor, this](TSharedPtr<IPixelStreaming2Streamer> Streamer) {
+			TSharedPtr<IPixelStreaming2InputHandler> Handler = Streamer->GetInputHandler().Pin();
+			{
+				UE_LOG(LogPixelStreaming2, Error, TEXT("Pixel Streaming input handler was null when sending response message."));
+				return;
+			}
+			Streamer->SendAllPlayersMessage(EPixelStreaming2FromStreamerMessage::Response, Descriptor);
+		});
+	}
+	else
+	{
+		UE_LOG(LogPixelStreaming2, Warning, TEXT("Pixel Streaming input component skipped sending response. This is expected on dedicated servers."));
+	}
+}
+
+void UPixelStreaming2Input::GetJsonStringValue(FString Descriptor, FString FieldName, FString& StringValue, bool& Success)
+{
+	UE::PixelStreaming2::ExtractJsonFromDescriptor(Descriptor, FieldName, StringValue, Success);
+}
+
+void UPixelStreaming2Input::AddJsonStringValue(const FString& Descriptor, FString FieldName, FString StringValue, FString& NewDescriptor, bool& Success)
+{
+	UE::PixelStreaming2::ExtendJsonWithField(Descriptor, FieldName, StringValue, NewDescriptor, Success);
+}
diff --git a/Engine/Plugins/Media/PixelStreaming2/Source/PixelStreaming2/Private/Blueprints/PixelStreaming2InputComponent.h b/Engine/Plugins/Media/PixelStreaming2/Source/PixelStreaming2/Private/Blueprints/PixelStreaming2InputComponent.h
new file mode 100644
--- /dev/null
+++ b/Engine/Plugins/Media/PixelStreaming2/Source/PixelStreaming2/Private/Blueprints/PixelStreaming2InputComponent.h
@@ -0,0 +1,65 @@
+// Copyright Epic Games, Inc. All Rights Reserved.
+#pragma once
+
+#include "Components/ActorComponent.h"
+#include "PixelStreaming2Module.h"
+#include "PixelStreaming2InputComponent.generated.h"
+
+/**
+ * This component may be attached to an actor to allow UI interactions to be
+ * handled as the delegate will be notified about the interaction and will be
+ * supplied with a generic descriptor string containing, for example, JSON data.
+ * Responses back to the source of the UI interactions may also be sent.
+ */
+UCLASS(Blueprintable, ClassGroup = (PixelStreaming2), meta = (BlueprintSpawnableComponent))
+class PIXELSTREAMING2_API UPixelStreaming2Input : public UActorComponent
+{
+	GENERATED_BODY()
+
+public:
+	UPixelStreaming2Input(const FObjectInitializer& ObjectInitializer);
+	void BeginPlay() override;
+	void EndPlay(const EEndPlayReason::Type EndPlayReason) override;
+
+	// The delegate which will be notified about a UI interaction.
+	DECLARE_DYNAMIC_MULTICAST_DELEGATE_OneParam(FOnInput, const FString&, Descriptor);
+	UPROPERTY(BlueprintAssignable, Category = "PixelStreaming2 Input")
+	FOnInput OnInputEvent;
+
+	/**
+	 * Send a response back to the source of the UI interactions.
+	 * @param Descriptor - A generic descriptor string.
+	 */
+	UFUNCTION(BlueprintCallable, Category = "PixelStreaming2 Input")
+	void SendPixelStreaming2Response(const FString& Descriptor);
+
+	/**
+	 * Helper function to extract a string field from a JSON descriptor of a
+	 * UI interaction given its field name.
+	 * The field name may be hierarchical, delimited by a period. For example,
+	 * to access the Width value of a Resolution command above you should use
+	 * "Resolution.Width" to get the width value.
+	 * @param Descriptor - The UI interaction JSON descriptor.
+	 * @param FieldName - The name of the field to look for in the JSON.
+	 * @param StringValue - The string value associated with the field name.
+	 * @param Success - True if the field exists in the JSON data.
+	 */
+	UFUNCTION(BlueprintPure, Category = "PixelStreaming2 Input")
+	void GetJsonStringValue(FString Descriptor, FString FieldName, FString& StringValue, bool& Success);
+
+	/**
+	 * Helper function to add a string field to a JSON descriptor. This produces
+	 * a new descriptor which may then be chained to add further string fields.
+	 * @param Descriptor - The initial JSON descriptor which may be blank initially.
+	 * @param FieldName - The name of the field to add to the JSON.
+	 * @param StringValue - The string value associated with the field name.
+	 * @param NewDescriptor - The JSON descriptor with the string field added.
+	 * @param Success - True if the string field could be added successfully.
+	 */
+	UFUNCTION(BlueprintPure, Category = "PixelStreaming2 Input")
+	void AddJsonStringValue(const FString& Descriptor, FString FieldName, FString StringValue, FString& NewDescriptor, bool& Success);
+
+private:
+	// For convenience we keep a reference to the Pixel Streaming plugin.
+	UE::PixelStreaming2::FPixelStreaming2Module* PixelStreaming2Module;
+};
diff --git a/Engine/Plugins/Media/PixelStreaming2/Source/PixelStreaming2/Private/Blueprints/PixelStreaming2MediaTexture.cpp b/Engine/Plugins/Media/PixelStreaming2/Source/PixelStreaming2/Private/Blueprints/PixelStreaming2MediaTexture.cpp
new file mode 100644
--- /dev/null
+++ b/Engine/Plugins/Media/PixelStreaming2/Source/PixelStreaming2/Private/Blueprints/PixelStreaming2MediaTexture.cpp
@@ -0,0 +1,165 @@
+// Copyright Epic Games, Inc. All Rights Reserved.
+
+#include "PixelStreaming2MediaTexture.h"
+
+#include "Async/Async.h"
+#include "DeviceProfiles/DeviceProfile.h"
+#include "DeviceProfiles/DeviceProfileManager.h"
+#include "TextureResource.h"
+
+constexpr int32 DEFAULT_WIDTH = 1920;
+constexpr int32 DEFAULT_HEIGHT = 1080;
+
+namespace UE::PixelStreaming2
+{
+	/**
+	 * The actual texture resource for a UPixelStreaming2MediaTexture. Contains the RHI Texture and
+	 * Sampler information.
+	 */
+	class FPixelStreaming2MediaTextureResource : public FTextureResource
+	{
+	public:
+		FPixelStreaming2MediaTextureResource(TWeakObjectPtr<UPixelStreaming2MediaTexture> Owner)
+			: MediaTexture(Owner)
+		{
+		}
+
+		virtual ~FPixelStreaming2MediaTextureResource() override
+		{
+			TextureRHI.SafeRelease();
+		}
+
+		virtual void InitRHI(FRHICommandListBase& RHICmdList) override
+		{
+			TStrongObjectPtr<UPixelStreaming2MediaTexture> PinnedMediaTexture = MediaTexture.Pin();
+			if (PinnedMediaTexture.IsValid())
+			{
+				FSamplerStateInitializerRHI SamplerStateInitializer(
+					(ESamplerFilter)UDeviceProfileManager::Get().GetActiveProfile()->GetTextureLODSettings()->GetSamplerFilter(PinnedMediaTexture.Get()),
+					AM_Border, AM_Border, AM_Wrap);
+				SamplerStateRHI = RHICreateSamplerState(SamplerStateInitializer);
+			}
+		}
+
+		virtual void ReleaseRHI() override
+		{
+			TextureRHI.SafeRelease();
+
+			TStrongObjectPtr<UPixelStreaming2MediaTexture> PinnedMediaTexture = MediaTexture.Pin();
+			if (PinnedMediaTexture.IsValid())
+			{
+				RHIUpdateTextureReference(PinnedMediaTexture->TextureReference.TextureReferenceRHI, nullptr);
+			}
+		}
+
+		virtual uint32 GetSizeX() const override
+		{
+			return TextureRHI.IsValid() ? TextureRHI->GetSizeXYZ().X : 0;
+		}
+
+		virtual uint32 GetSizeY() const override
+		{
+			return TextureRHI.IsValid() ? TextureRHI->GetSizeXYZ().Y : 0;
+		}
+
+		SIZE_T GetResourceSize()
+		{
+			return CalcTextureSize(GetSizeX(), GetSizeY(), EPixelFormat::PF_A8R8G8B8, 1);
+		}
+
+	private:
+		TWeakObjectPtr<UPixelStreaming2MediaTexture> MediaTexture;
+	};
+} // namespace UE::PixelStreaming2
+
+UPixelStreaming2MediaTexture::UPixelStreaming2MediaTexture(const FObjectInitializer& ObjectInitializer)
+	: Super(ObjectInitializer)
+{
+	SetResource(nullptr);
+}
+
+void UPixelStreaming2MediaTexture::BeginDestroy()
+{
+	SetResource(nullptr);
+
+	Super::BeginDestroy();
+}
+
+void UPixelStreaming2MediaTexture::GetResourceSizeEx(FResourceSizeEx& CumulativeResourceSize)
+{
+	Super::GetResourceSizeEx(CumulativeResourceSize);
+
+	if (CurrentResource)
+	{
+		CumulativeResourceSize.AddUnknownMemoryBytes(CurrentResource->GetResourceSize());
+	}
+}
+
+FTextureResource* UPixelStreaming2MediaTexture::CreateResource()
+{
+	if (CurrentResource)
+	{
+		SetResource(nullptr);
+		CurrentResource = nullptr;
+	}
+
+	CurrentResource = new UE::PixelStreaming2::FPixelStreaming2MediaTextureResource(this);
+	InitializeResources();
+
+	return CurrentResource;
+}
+
+void UPixelStreaming2MediaTexture::ConsumeFrame(FTextureRHIRef Frame)
+{
+	AsyncTask(ENamedThreads::ActualRenderingThread, [this, Frame]() {
+		FScopeLock Lock(&RenderSyncContext);
+
+		FRHICommandListImmediate& RHICmdList = FRHICommandListImmediate::Get();
+		UpdateTextureReference(RHICmdList, Frame);
+	});
+}
+
+void UPixelStreaming2MediaTexture::InitializeResources()
+{
+	ENQUEUE_RENDER_COMMAND(FPixelStreamingMediaTextureUpdateTextureReference)
+	([this](FRHICommandListImmediate& RHICmdList) {
+		// Set the default video texture to reference nothing
+		FTextureRHIRef ShaderTexture2D;
+		FTextureRHIRef RenderableTexture;
+
+		FRHITextureCreateDesc RenderTargetTextureDesc =
+			FRHITextureCreateDesc::Create2D(TEXT(""), DEFAULT_WIDTH, DEFAULT_HEIGHT, EPixelFormat::PF_B8G8R8A8)
+				.SetClearValue(FClearValueBinding(FLinearColor(0.0f, 0.0f, 0.0f)))
+				.SetFlags(ETextureCreateFlags::Dynamic | ETextureCreateFlags::ShaderResource | TexCreate_RenderTargetable)
+				.SetInitialState(ERHIAccess::SRVMask);
+
+		RenderableTexture = RHICmdList.CreateTexture(RenderTargetTextureDesc);
+		ShaderTexture2D = RenderableTexture;
+
+		CurrentResource->TextureRHI = RenderableTexture;
+
+		RHICmdList.UpdateTextureReference(TextureReference.TextureReferenceRHI, CurrentResource->TextureRHI);
+	});
+}
+
+void UPixelStreaming2MediaTexture::UpdateTextureReference(FRHICommandList& RHICmdList, FTextureRHIRef Reference)
+{
+	if (CurrentResource)
+	{
+		if (Reference.IsValid() && CurrentResource->TextureRHI != Reference)
+		{
+			CurrentResource->TextureRHI = Reference;
+			RHICmdList.UpdateTextureReference(TextureReference.TextureReferenceRHI, CurrentResource->TextureRHI);
+		}
+		else if (!Reference.IsValid())
+		{
+			if (CurrentResource)
+			{
+				InitializeResources();
+
+				// Make sure RenderThread is executed before continuing
+				FlushRenderingCommands();
+			}
+		}
+	}
+}
diff --git a/Engine/Plugins/Media/PixelStreaming2/Source/PixelStreaming2/Private/Blueprints/PixelStreaming2Peer.cpp b/Engine/Plugins/Media/PixelStreaming2/Source/PixelStreaming2/Private/Blueprints/PixelStreaming2Peer.cpp
new file mode 100644
--- /dev/null
+++ b/Engine/Plugins/Media/PixelStreaming2/Source/PixelStreaming2/Private/Blueprints/PixelStreaming2Peer.cpp
@@ -0,0 +1,417 @@
+// Copyright Epic Games, Inc. All Rights Reserved.
+
+#include "PixelStreaming2Peer.h"
+
+#include "Async/Async.h"
+#include "EpicRtcSessionObserver.h"
+#include "Logging.h"
+#include "PixelStreaming2Module.h"
+#include "SampleBuffer.h"
+#include "SoundGenerator.h"
+#include "ToStringExtensions.h"
+
+uint32_t UPixelStreaming2Peer::PlayerId = 0;
+
+UPixelStreaming2Peer::UPixelStreaming2Peer(const FObjectInitializer& ObjectInitializer)
+	: Super(ObjectInitializer)
+	, SoundGenerator(MakeShared<UE::PixelStreaming2::FSoundGenerator, ESPMode::ThreadSafe>())
+{
+	PreferredBufferLength = 512u;
+	NumChannels = 2;
+	PrimaryComponentTick.bCanEverTick = true;
+	SetComponentTickEnabled(true);
+	bAutoActivate = true;
+
+	PlayerName = FUtf8String::Printf("PixelStreaming2Player%d", PlayerId++);
+}
+
+void UPixelStreaming2Peer::BeginPlay()
+{
+	EpicRtcManager = MakeShared<UE::PixelStreaming2::FEpicRtcManager>();
+
+	EpicRtcManager->OnSessionStateUpdate.AddUObject(this, &UPixelStreaming2Peer::OnSessionStateUpdate);
+	EpicRtcManager->OnSessionErrorUpdate.AddUObject(this, &UPixelStreaming2Peer::OnSessionErrorUpdate);
+	EpicRtcManager->OnSessionRoomsAvailableUpdate.AddUObject(this, &UPixelStreaming2Peer::OnSessionRoomsAvailableUpdate);
+
+	EpicRtcManager->OnRoomStateUpdate.AddUObject(this, &UPixelStreaming2Peer::OnRoomStateUpdate);
+	EpicRtcManager->OnRoomJoinedUpdate.AddUObject(this, &UPixelStreaming2Peer::OnRoomJoinedUpdate);
+	EpicRtcManager->OnRoomLeftUpdate.AddUObject(this, &UPixelStreaming2Peer::OnRoomLeftUpdate);
+	EpicRtcManager->OnAudioTrackUpdate.AddUObject(this, &UPixelStreaming2Peer::OnAudioTrackUpdate);
+	EpicRtcManager->OnVideoTrackUpdate.AddUObject(this, &UPixelStreaming2Peer::OnVideoTrackUpdate);
+	EpicRtcManager->OnDataTrackUpdate.AddUObject(this, &UPixelStreaming2Peer::OnDataTrackUpdate);
+	EpicRtcManager->OnLocalSdpUpdate.AddUObject(this, &UPixelStreaming2Peer::OnLocalSdpUpdate);
+	EpicRtcManager->OnRemoteSdpUpdate.AddUObject(this, &UPixelStreaming2Peer::OnRemoteSdpUpdate);
+	EpicRtcManager->OnRoomErrorUpdate.AddUObject(this, &UPixelStreaming2Peer::OnRoomErrorUpdate);
+
+	EpicRtcManager->OnAudioTrackMuted.AddUObject(this, &UPixelStreaming2Peer::OnAudioTrackMuted);
+	EpicRtcManager->OnAudioTrackFrame.AddUObject(this, &UPixelStreaming2Peer::OnAudioTrackFrame);
+	EpicRtcManager->OnAudioTrackRemoved.AddUObject(this, &UPixelStreaming2Peer::OnAudioTrackRemoved);
+	EpicRtcManager->OnAudioTrackState.AddUObject(this, &UPixelStreaming2Peer::OnAudioTrackState);
+
+	EpicRtcManager->OnVideoTrackMuted.AddUObject(this, &UPixelStreaming2Peer::OnVideoTrackMuted);
+	EpicRtcManager->OnVideoTrackFrame.AddUObject(this, &UPixelStreaming2Peer::OnVideoTrackFrame);
+	EpicRtcManager->OnVideoTrackRemoved.AddUObject(this, &UPixelStreaming2Peer::OnVideoTrackRemoved);
+	EpicRtcManager->OnVideoTrackState.AddUObject(this, &UPixelStreaming2Peer::OnVideoTrackState);
+
+	EpicRtcManager->OnDataTrackRemoved.AddUObject(this, &UPixelStreaming2Peer::OnDataTrackRemoved);
+	EpicRtcManager->OnDataTrackState.AddUObject(this, &UPixelStreaming2Peer::OnDataTrackState);
+	EpicRtcManager->OnDataTrackMessage.AddUObject(this, &UPixelStreaming2Peer::OnDataTrackMessage);
+
+	EpicRtcManager->EpicRtcConference = UE::PixelStreaming2::FPixelStreaming2Module::GetModule()->GetEpicRtcConference();
+
+	EpicRtcManager->SessionObserver = new UE::PixelStreaming2::FEpicRtcSessionObserver(EpicRtcManager);
+	EpicRtcManager->RoomObserver = new UE::PixelStreaming2::FEpicRtcRoomObserver(EpicRtcManager);
+
+	EpicRtcManager->AudioTrackObserverFactory = new UE::PixelStreaming2::FEpicRtcAudioTrackObserverFactory(EpicRtcManager);
+	EpicRtcManager->VideoTrackObserverFactory = new UE::PixelStreaming2::FEpicRtcVideoTrackObserverFactory(EpicRtcManager);
+	EpicRtcManager->DataTrackObserverFactory = new UE::PixelStreaming2::FEpicRtcDataTrackObserverFactory(EpicRtcManager);
+
+	Super::BeginPlay();
+}
+
+void UPixelStreaming2Peer::EndPlay(const EEndPlayReason::Type EndPlayReason)
+{
+	Disconnect();
+
+	Super::EndPlay(EndPlayReason);
+}
+
+bool UPixelStreaming2Peer::Connect(const FString& Url)
+{
+	if (!EpicRtcManager)
+	{
+		UE_LOGFMT(LogPixelStreaming2, Warning, "Failed to connect. EpicRtcManager isn't valid!");
+		return false;
+	}
+
+	FUtf8String Utf8Url = *Url;
+	FUtf8String ConnectionUrl = Utf8Url + (Utf8Url.Contains(TEXT("?")) ? TEXT("&") : TEXT("?")) + TEXT("isStreamer=false");
+
+	EpicRtcSessionConfig SessionConfig = {
+		._id = { ._ptr = reinterpret_cast<const char*>(*PlayerName), ._length = static_cast<uint64_t>(PlayerName.Len()) },
+		._url = { ._ptr = reinterpret_cast<const char*>(*ConnectionUrl), ._length = static_cast<uint64_t>(ConnectionUrl.Len()) },
+		._observer = EpicRtcManager->SessionObserver.GetReference()
+	};
+
+	EpicRtcErrorCode Result = EpicRtcManager->EpicRtcConference->CreateSession(SessionConfig, EpicRtcManager->EpicRtcSession.GetInitReference());
+	if (Result != EpicRtcErrorCode::Ok)
+	{
+		UE_LOGFMT(LogPixelStreaming2, Error, "Failed to create EpicRtc session, CreateSession returned: {0}", UE::PixelStreaming2::ToString(Result));
+		return false;
+	}
+
+	Result = EpicRtcManager->EpicRtcSession->Connect();
+	if (Result != EpicRtcErrorCode::Ok)
+	{
+		UE_LOGFMT(LogPixelStreaming2, Error, "Failed to connect EpicRtcSession. Connect returned: {0}", UE::PixelStreaming2::ToString(Result));
+		return false;
+	}
+	return true;
+}
+
+bool UPixelStreaming2Peer::Disconnect()
+{
+	if (!EpicRtcManager)
+	{
+		UE_LOGFMT(LogPixelStreaming2, Warning, "Failed to disconnect. EpicRtcManager isn't valid!");
+		return false;
+	}
+
+	if (!EpicRtcManager->EpicRtcSession)
+	{
+		UE_LOGFMT(LogPixelStreaming2, Error, "Failed to disconnect, EpicRtcSession does not exist");
+		return false;
+	}
+
+	if (EpicRtcManager->EpicRtcRoom)
+	{
+		EpicRtcManager->EpicRtcRoom->Leave();
+		EpicRtcManager->EpicRtcSession->RemoveRoom({ ._ptr = reinterpret_cast<const char*>(*SubscribedStream), ._length = static_cast<uint64_t>(SubscribedStream.Len()) });
+	}
+
+	EpicRtcErrorCode Result = EpicRtcManager->EpicRtcSession->Disconnect();
+	if (Result != EpicRtcErrorCode::Ok)
+	{
+		UE_LOGFMT(LogPixelStreaming2, Error, "Failed to disconnect EpicRtcSession. Disconnect returned {0}", UE::PixelStreaming2::ToString(Result));
+		return false;
+	}
+
+	if (AudioSink)
+	{
+		AudioSink->RemoveAudioConsumer(this);
+	}
+
+	if (VideoSink && VideoConsumer)
+	{
+		VideoSink->RemoveVideoConsumer(VideoConsumer);
+	}
+
+	EpicRtcManager->EpicRtcConference->RemoveSession({ ._ptr = reinterpret_cast<const char*>(*PlayerName), ._length = static_cast<uint64_t>(PlayerName.Len()) });
+
+	return true;
+}
+
+bool UPixelStreaming2Peer::Subscribe(const FString& StreamerId)
+{
+	if (!EpicRtcManager)
+	{
+		UE_LOGFMT(LogPixelStreaming2, Warning, "Failed to subscribe. EpicRtcManager isn't valid!");
+		return false;
+	}
+
+	if (SessionState != EpicRtcSessionState::Connected)
+	{
+		UE_LOGFMT(LogPixelStreaming2, Error, "Failed to create subscribe to streamer. EpicRtc session isn't connected");
+		return false;
+	}
+
+	EpicRtcConnectionConfig ConnectionConfig = {
+		// TODO (Migration): RTCP-7032 This info usually comes from the OnConfig signalling message
+		._iceServers = EpicRtcIceServerSpan{ ._ptr = nullptr, ._size = 0 },
+		._iceConnectionPolicy = EpicRtcIcePolicy::All,
+		._disableTcpCandidates = false
+	};
+
+	SubscribedStream = *StreamerId;
+
+	EpicRtcRoomConfig RoomConfig = {
+		._id = EpicRtcStringView{ ._ptr = reinterpret_cast<const char*>(*SubscribedStream), ._length = static_cast<uint64_t>(SubscribedStream.Len()) },
+		._connectionConfig = ConnectionConfig,
+		._ticket = EpicRtcStringView{ ._ptr = nullptr, ._length = 0 },
+		._observer = EpicRtcManager->RoomObserver,
+		._audioTrackObserverFactory = EpicRtcManager->AudioTrackObserverFactory,
+		._dataTrackObserverFactory = EpicRtcManager->DataTrackObserverFactory,
+		._videoTrackObserverFactory = EpicRtcManager->VideoTrackObserverFactory
+	};
+
+	EpicRtcErrorCode Result = EpicRtcManager->EpicRtcSession->CreateRoom(RoomConfig, EpicRtcManager->EpicRtcRoom.GetInitReference());
+	if (Result != EpicRtcErrorCode::Ok)
+	{
+		UE_LOGFMT(LogPixelStreaming2, Error, "Failed to create EpicRtc room. CreateRoom returned: {0}", UE::PixelStreaming2::ToString(Result));
+		return false;
+	}
+
+	EpicRtcManager->EpicRtcRoom->Join();
+
+	return true;
+}
+
+ISoundGeneratorPtr UPixelStreaming2Peer::CreateSoundGenerator(const FSoundGeneratorInitParams& InParams)
+{
+	SoundGenerator->SetParameters(InParams);
+	Initialize(InParams.SampleRate);
+
+	return SoundGenerator;
+}
+
+void UPixelStreaming2Peer::OnBeginGenerate()
+{
+	SoundGenerator->bGeneratingAudio = true;
+}
+
+void UPixelStreaming2Peer::OnEndGenerate()
+{
+	SoundGenerator->bGeneratingAudio = false;
+}
+
+void UPixelStreaming2Peer::ConsumeRawPCM(const int16_t* AudioData, int InSampleRate, size_t NChannels, size_t NFrames)
+{
+	// Sound generator has not been initialized yet.
+	if (!SoundGenerator || GetAudioComponent() == nullptr)
+	{
+		return;
+	}
+
+	// Set pitch multiplier as a way to handle mismatched sample rates
+	if (InSampleRate != SoundGenerator->GetSampleRate())
+	{
+		GetAudioComponent()->SetPitchMultiplier((float)InSampleRate / SoundGenerator->GetSampleRate());
+	}
+	else if (GetAudioComponent()->PitchMultiplier != 1.0f)
+	{
+		GetAudioComponent()->SetPitchMultiplier(1.0f);
+	}
+
+	Audio::TSampleBuffer<int16_t> Buffer(AudioData, NFrames, NChannels, InSampleRate);
+	if (NChannels != SoundGenerator->GetNumChannels())
+	{
+		Buffer.MixBufferToChannels(SoundGenerator->GetNumChannels());
+	}
+
+	SoundGenerator->AddAudio(Buffer.GetData(), InSampleRate, NChannels, Buffer.GetNumSamples());
+}
+
+void UPixelStreaming2Peer::OnConsumerAdded()
+{
+	SoundGenerator->bShouldGenerateAudio = true;
+}
+
+void UPixelStreaming2Peer::OnConsumerRemoved()
+{
+	SoundGenerator->bShouldGenerateAudio = false;
+	SoundGenerator->EmptyBuffers();
+}
+
+void UPixelStreaming2Peer::OnSessionStateUpdate(const EpicRtcSessionState StateUpdate)
+{
+	switch (StateUpdate)
+	{
+		case EpicRtcSessionState::New:
+		case EpicRtcSessionState::Pending:
+		case EpicRtcSessionState::Connected:
+		case EpicRtcSessionState::Disconnected:
+		case EpicRtcSessionState::Failed:
+		case EpicRtcSessionState::Exiting:
+			SessionState = StateUpdate;
+			break;
+		default:
+			UE_LOGFMT(LogPixelStreaming2, Warning, "OnSessionStateUpdate received unknown EpicRtcSessionState: {0}", static_cast<int>(StateUpdate));
+			break;
+	}
+}
+
+void UPixelStreaming2Peer::OnSessionRoomsAvailableUpdate(EpicRtcStringArrayInterface* RoomsList)
+{
+	TArray<FString> Streamers;
+
+	for (uint64_t i = 0; i < RoomsList->Size(); i++)
+	{
+		EpicRtcStringInterface* RoomName = RoomsList->Get()[i];
+		Streamers.Add(FString{ RoomName->Get(), (int32)RoomName->Length() });
+	}
+
+	OnStreamerList.Broadcast(Streamers);
+}
+
+void UPixelStreaming2Peer::OnSessionErrorUpdate(const EpicRtcErrorCode ErrorUpdate)
+{
+}
+
+void UPixelStreaming2Peer::OnRoomStateUpdate(const EpicRtcRoomState State)
+{
+}
+
+void UPixelStreaming2Peer::OnRoomJoinedUpdate(EpicRtcParticipantInterface* Participant)
+{
+	FString ParticipantId{ (int32)Participant->GetId()._length, Participant->GetId()._ptr };
+	UE_LOG(LogPixelStreaming2, Log, TEXT("Player (%s) joined"), *ParticipantId);
+}
+
+void UPixelStreaming2Peer::OnRoomLeftUpdate(const EpicRtcStringView ParticipantId)
+{
+}
+
+void UPixelStreaming2Peer::OnAudioTrackUpdate(EpicRtcParticipantInterface* Participant, EpicRtcAudioTrackInterface* AudioTrack)
+{
+	FString ParticipantId{ (int32)Participant->GetId()._length, Participant->GetId()._ptr };
+	FString VideoTrackId{ (int32)AudioTrack->GetId()._length, AudioTrack->GetId()._ptr };
+	UE_LOG(LogPixelStreaming2, VeryVerbose, TEXT("UPixelStreaming2Peer::OnAudioTrackUpdate(Participant [%s], VideoTrack [%s])"), *ParticipantId, *VideoTrackId);
+
+	if (AudioTrack->IsRemote())
+	{
+		// We received a remote track. We should now generate audio from it
+		AudioSink = MakeShared<UE::PixelStreaming2::FAudioSink>();
+		AudioSink->AddAudioConsumer(this);
+	}
+}
+
+void UPixelStreaming2Peer::OnVideoTrackUpdate(EpicRtcParticipantInterface* Participant, EpicRtcVideoTrackInterface* VideoTrack)
+{
+	FString ParticipantId{ (int32)Participant->GetId()._length, Participant->GetId()._ptr };
+	FString VideoTrackId{ (int32)VideoTrack->GetId()._length, VideoTrack->GetId()._ptr };
+	UE_LOG(LogPixelStreaming2, VeryVerbose, TEXT("UPixelStreaming2Peer::OnVideoTrackUpdate(Participant [%s], VideoTrack [%s])"), *ParticipantId, *VideoTrackId);
+
+	if (VideoTrack->IsRemote())
+	{
+		// We received a remote track. We should now generate video from it
+		VideoSink = MakeShared<UE::PixelStreaming2::FVideoSink>();
+		if (VideoConsumer)
+		{
+			VideoSink->AddVideoConsumer(VideoConsumer);
+		}
+	}
+}
+
+void UPixelStreaming2Peer::OnDataTrackUpdate(EpicRtcParticipantInterface* Participant, EpicRtcDataTrackInterface* DataTrack)
+{
+}
+
+void UPixelStreaming2Peer::OnLocalSdpUpdate(EpicRtcParticipantInterface* Participant, EpicRtcSdpInterface* Sdp)
+{
+}
+
+void UPixelStreaming2Peer::OnRemoteSdpUpdate(EpicRtcParticipantInterface* Participant, EpicRtcSdpInterface* Sdp)
+{
+}
+
+void UPixelStreaming2Peer::OnRoomErrorUpdate(const EpicRtcErrorCode Error)
+{
+}
+
+void UPixelStreaming2Peer::OnAudioTrackMuted(EpicRtcAudioTrackInterface* AudioTrack, EpicRtcBool bIsMuted)
+{
+	if (!AudioSink)
+	{
+		return;
+	}
+
+	AudioSink->SetMuted((bool)bIsMuted);
+}
+
+void UPixelStreaming2Peer::OnAudioTrackFrame(EpicRtcAudioTrackInterface* AudioTrack, const EpicRtcAudioFrame& Frame)
+{
+	if (!AudioSink)
+	{
+		return;
+	}
+
+	AudioSink->OnAudioData(Frame._data, Frame._format._sampleRate, Frame._format._numChannels, Frame._length);
+}
+
+void UPixelStreaming2Peer::OnAudioTrackRemoved(EpicRtcAudioTrackInterface* AudioTrack)
+{
+}
+
+void UPixelStreaming2Peer::OnAudioTrackState(EpicRtcAudioTrackInterface* AudioTrack, const EpicRtcTrackState State)
+{
+}
+
+void UPixelStreaming2Peer::OnVideoTrackMuted(EpicRtcVideoTrackInterface* VideoTrack, EpicRtcBool bIsMuted)
+{
+	if (!VideoSink)
+	{
+		return;
+	}
+
+	VideoSink->SetMuted((bool)bIsMuted);
+}
+
+void UPixelStreaming2Peer::OnVideoTrackFrame(EpicRtcVideoTrackInterface* VideoTrack, const EpicRtcVideoFrame& Frame)
+{
+	if (!VideoSink)
+	{
+		return;
+	}
+
+	VideoSink->OnVideoData(Frame);
+}
+
+void UPixelStreaming2Peer::OnVideoTrackRemoved(EpicRtcVideoTrackInterface* VideoTrack)
+{
+}
+
+void UPixelStreaming2Peer::OnVideoTrackState(EpicRtcVideoTrackInterface* VideoTrack, const EpicRtcTrackState State)
+{
+}
+
+void UPixelStreaming2Peer::OnDataTrackRemoved(EpicRtcDataTrackInterface* DataTrack)
+{
+}
+
+void UPixelStreaming2Peer::OnDataTrackState(EpicRtcDataTrackInterface* DataTrack, const EpicRtcTrackState State)
+{
+}
+
+void UPixelStreaming2Peer::OnDataTrackMessage(EpicRtcDataTrackInterface* DataTrack)
+{
+}
\ No newline at end of file
diff --git a/Engine/Plugins/Media/PixelStreaming2/Source/PixelStreaming2/Private/Blueprints/PixelStreaming2Peer.h b/Engine/Plugins/Media/PixelStreaming2/Source/PixelStreaming2/Private/Blueprints/PixelStreaming2Peer.h
new file mode 100644
--- /dev/null
+++ b/Engine/Plugins/Media/PixelStreaming2/Source/PixelStreaming2/Private/Blueprints/PixelStreaming2Peer.h
@@ -0,0 +1,128 @@
+// Copyright Epic Games, Inc. All Rights Reserved.
+
+#pragma once
+
+#include "Components/SynthComponent.h"
+#include "Containers/Utf8String.h"
+#include "IPixelStreaming2AudioConsumer.h"
+#include "AudioSink.h"
+#include "EpicRtcManager.h"
+#include "PixelStreaming2MediaTexture.h"
+#include "VideoSink.h"
+
+#include "PixelStreaming2Peer.generated.h"
+
+DECLARE_DYNAMIC_MULTICAST_SPARSE_DELEGATE_OneParam(FPixelStreamingStreamerList, UPixelStreaming2Peer, OnStreamerList, const TArray<FString>&, StreamerList);
+
+namespace UE::PixelStreaming2
+{
+	class FSoundGenerator;
+} // namespace UE::PixelStreaming2
+
+/**
+ * A blueprint representation of a Pixel Streaming Peer Connection. Will accept video sinks to receive video data.
+ * NOTE: This class is not a peer of a streamer. This class represents a peer in its own right (akin to the browser) and will subscribe to a stream
+ */
+// UCLASS(Blueprintable, ClassGroup = (PixelStreaming2), meta = (BlueprintSpawnableComponent))
+UCLASS(Blueprintable, Category = "PixelStreaming2", META = (DisplayName = "PixelStreaming Peer Component", BlueprintSpawnableComponent))
+class UPixelStreaming2Peer : public USynthComponent, public IPixelStreaming2AudioConsumer
+{
+	GENERATED_UCLASS_BODY()
+protected:
+	// UObject overrides
+	virtual void BeginPlay() override;
+	virtual void EndPlay(const EEndPlayReason::Type EndPlayReason) override;
+
+	// USynthComponent overrides
+	virtual void			   OnBeginGenerate() override;
+	virtual void			   OnEndGenerate() override;
+	virtual ISoundGeneratorPtr CreateSoundGenerator(const FSoundGeneratorInitParams& InParams) override;
+
+public:
+	/**
+	 * Attempt to connect to a specified signalling server.
+	 * @param Url The url of the signalling server. Ignored if this component has a MediaSource. In that case the URL on the media source will be used instead.
+	 * @return Returns false if Connect fails.
+	 */
+	UFUNCTION(BlueprintCallable, Category = "PixelStreaming2")
+	bool Connect(const FString& Url);
+
+	/**
+	 * Disconnect from the signalling server. No action if no connection exists.
+	 * @return Returns false if Disconnect fails.
+	 */
+	UFUNCTION(BlueprintCallable, Category = "PixelStreaming2")
+	bool Disconnect();
+
+	/**
+	 * Subscribe this peer to the streams provided by the specified streamer.
+	 * @param StreamerId The name of the streamer to subscribe to.
+	 * @return Returns false if Subscribe fails.
+	 */
+	UFUNCTION(BlueprintCallable, Category = "PixelStreaming2")
+	bool Subscribe(const FString& StreamerId);
+
+	/**
+	 * Fired when the connection the list of available streams from the server.
+	 */
+	UPROPERTY(BlueprintAssignable, Category = "Components|Activation")
+	FPixelStreamingStreamerList OnStreamerList;
+
+	/**
+	 * A sink for the video data received once this connection has finished negotiating.
+	 */
+	UPROPERTY(EditDefaultsOnly, Category = "Properties", META = (DisplayName = "Pixel Streaming Video Consumer", AllowPrivateAccess = true))
+	TObjectPtr<UPixelStreaming2MediaTexture> VideoConsumer = nullptr;
+
+private:
+	FUtf8String		SubscribedStream;
+	FUtf8String		PlayerName;
+	static uint32_t PlayerId;
+
+	TSharedPtr<UE::PixelStreaming2::FAudioSink>							AudioSink;
+	TSharedPtr<UE::PixelStreaming2::FSoundGenerator, ESPMode::ThreadSafe> SoundGenerator;
+
+	TSharedPtr<UE::PixelStreaming2::FVideoSink> VideoSink;
+
+	TSharedPtr<UE::PixelStreaming2::FEpicRtcManager> EpicRtcManager;
+
+public:
+	/* Begin IPixelStreaming2AudioConsumer Callbacks */
+	virtual void ConsumeRawPCM(const int16_t* AudioData, int InSampleRate, size_t NChannels, size_t NFrames) override;
+	virtual void OnConsumerAdded() override;
+	virtual void OnConsumerRemoved() override;
+	/* End IPixelStreaming2AudioConsumer Callbacks */
+
+	/* Begin FEpicRtcManager Callbacks */
+	void OnSessionStateUpdate(const EpicRtcSessionState StateUpdate);
+	void OnSessionErrorUpdate(const EpicRtcErrorCode ErrorUpdate);
+	void OnSessionRoomsAvailableUpdate(EpicRtcStringArrayInterface* RoomsList);
+
+	void OnRoomStateUpdate(const EpicRtcRoomState State);
+	void OnRoomJoinedUpdate(EpicRtcParticipantInterface* Participant);
+	void OnRoomLeftUpdate(const EpicRtcStringView ParticipantId);
+	void OnAudioTrackUpdate(EpicRtcParticipantInterface* Participant, EpicRtcAudioTrackInterface* AudioTrack);
+	void OnVideoTrackUpdate(EpicRtcParticipantInterface* Participant, EpicRtcVideoTrackInterface* VideoTrack);
+	void OnDataTrackUpdate(EpicRtcParticipantInterface* Participant, EpicRtcDataTrackInterface* DataTrack);
+	void OnLocalSdpUpdate(EpicRtcParticipantInterface* Participant, EpicRtcSdpInterface* Sdp);
+	void OnRemoteSdpUpdate(EpicRtcParticipantInterface* Participant, EpicRtcSdpInterface* Sdp);
+	void OnRoomErrorUpdate(const EpicRtcErrorCode Error);
+
+	void OnAudioTrackMuted(EpicRtcAudioTrackInterface* AudioTrack, EpicRtcBool bIsMuted);
+	void OnAudioTrackFrame(EpicRtcAudioTrackInterface* AudioTrack, const EpicRtcAudioFrame& Frame);
+	void OnAudioTrackRemoved(EpicRtcAudioTrackInterface* AudioTrack);
+	void OnAudioTrackState(EpicRtcAudioTrackInterface* AudioTrack, const EpicRtcTrackState State);
+
+	void OnVideoTrackMuted(EpicRtcVideoTrackInterface* VideoTrack, EpicRtcBool bIsMuted);
+	void OnVideoTrackFrame(EpicRtcVideoTrackInterface* VideoTrack, const EpicRtcVideoFrame& Frame);
+	void OnVideoTrackRemoved(EpicRtcVideoTrackInterface* VideoTrack);
+	void OnVideoTrackState(EpicRtcVideoTrackInterface* VideoTrack, const EpicRtcTrackState State);
+
+	void OnDataTrackRemoved(EpicRtcDataTrackInterface* DataTrack);
+	void OnDataTrackState(EpicRtcDataTrackInterface* DataTrack, const EpicRtcTrackState State);
+	void OnDataTrackMessage(EpicRtcDataTrackInterface* DataTrack);
+	/* End FEpicRtcManager Callbacks */
+
+private:
+	EpicRtcSessionState SessionState = EpicRtcSessionState::Disconnected;
+};
diff --git a/Engine/Plugins/Media/PixelStreaming2/Source/PixelStreaming2/Private/Blueprints/PixelStreaming2VideoComponent.cpp b/Engine/Plugins/Media/PixelStreaming2/Source/PixelStreaming2/Private/Blueprints/PixelStreaming2VideoComponent.cpp
new file mode 100644
--- /dev/null
+++ b/Engine/Plugins/Media/PixelStreaming2/Source/PixelStreaming2/Private/Blueprints/PixelStreaming2VideoComponent.cpp
@@ -0,0 +1,146 @@
+// Copyright Epic Games, Inc. All Rights Reserved.
+
+#include "PixelStreaming2VideoComponent.h"
+
+#include "CoreMinimal.h"
+#include "IPixelStreaming2Module.h"
+#include "IPixelStreaming2Streamer.h"
+#include "Logging.h"
+
+/*
+ * Component that recieves video from a remote webrtc connection and outputs it into UE using a "synth component".
+ */
+UPixelStreaming2VideoComponent::UPixelStreaming2VideoComponent(const FObjectInitializer& ObjectInitializer)
+	: Super(ObjectInitializer)
+	, PlayerToWatch(FString())
+	, bAutoFindPeer(true)
+	, VideoSink(nullptr)
+	, bIsWatchingPlayer(false)
+{
+	PrimaryComponentTick.bCanEverTick = true;
+	SetComponentTickEnabled(true);
+	bAutoActivate = true;
+};
+
+void UPixelStreaming2VideoComponent::BeginDestroy()
+{
+	Super::BeginDestroy();
+	Reset();
+}
+
+bool UPixelStreaming2VideoComponent::Watch(FString PlayerId)
+{
+	IPixelStreaming2Module& PixelStreaming2Module = IPixelStreaming2Module::Get();
+	if (!PixelStreaming2Module.IsReady())
+	{
+		return false;
+	}
+	return StreamerWatch(PixelStreaming2Module.GetDefaultStreamerID(), PlayerId);
+}
+
+bool UPixelStreaming2VideoComponent::StreamerWatch(FString StreamerId, FString PlayerId)
+{
+	if (!IPixelStreaming2Module::IsAvailable())
+	{
+		UE_LOG(LogPixelStreaming2, Verbose, TEXT("Pixel Streaming video component could not watch anything because Pixel Streaming module is not loaded. This is expected on dedicated servers."));
+		return false;
+	}
+
+	IPixelStreaming2Module& PixelStreaming2Module = IPixelStreaming2Module::Get();
+	if (!PixelStreaming2Module.IsReady())
+	{
+		return false;
+	}
+
+	PlayerToWatch = PlayerId;
+
+	if (StreamerId == FString())
+	{
+		TArray<FString> StreamerIds = PixelStreaming2Module.GetStreamerIds();
+		if (StreamerIds.Num() > 0)
+		{
+			StreamerToWatch = StreamerIds[0];
+		}
+		else
+		{
+			StreamerToWatch = PixelStreaming2Module.GetDefaultStreamerID();
+		}
+	}
+	else
+	{
+		StreamerToWatch = StreamerId;
+	}
+
+	TSharedPtr<IPixelStreaming2Streamer> Streamer = PixelStreaming2Module.FindStreamer(StreamerToWatch);
+	if (!Streamer)
+	{
+		return false;
+	}
+
+	IPixelStreaming2VideoSink* CandidateSink = WillWatchAnyPlayer() ? Streamer->GetUnwatchedVideoSink() : Streamer->GetPeerVideoSink(FString(PlayerToWatch));
+
+	if (CandidateSink == nullptr)
+	{
+		return false;
+	}
+
+	VideoSink = CandidateSink;
+
+	if (VideoConsumer)
+	{
+		VideoSink->AddVideoConsumer(VideoConsumer);
+		bIsWatchingPlayer = true;
+	}
+
+	return true;
+}
+
+void UPixelStreaming2VideoComponent::Reset()
+{
+	PlayerToWatch = FString();
+	StreamerToWatch = FString();
+	if (VideoSink && VideoConsumer)
+	{
+		VideoSink->RemoveVideoConsumer(VideoConsumer);
+		bIsWatchingPlayer = false;
+	}
+	VideoSink = nullptr;
+}
+
+bool UPixelStreaming2VideoComponent::IsWatchingPlayer()
+{
+	return bIsWatchingPlayer;
+}
+
+bool UPixelStreaming2VideoComponent::WillWatchAnyPlayer()
+{
+	return PlayerToWatch == FString();
+}
+
+void UPixelStreaming2VideoComponent::TickComponent(float DeltaTime, ELevelTick TickType, FActorComponentTickFunction* ThisTickFunction)
+{
+	bool bPixelStreaming2Loaded = IPixelStreaming2Module::IsAvailable();
+
+	// Early out if running in commandlet
+	if (IsRunningCommandlet())
+	{
+		return;
+	}
+
+	// if auto connect turned off don't bother
+	if (!bAutoFindPeer)
+	{
+		return;
+	}
+
+	// if watching a peer don't auto connect
+	if (IsWatchingPlayer())
+	{
+		return;
+	}
+
+	if (StreamerWatch(StreamerToWatch, PlayerToWatch))
+	{
+		UE_LOG(LogPixelStreaming2, Log, TEXT("PixelStreaming2 video component found a WebRTC peer to watch."));
+	}
+}
diff --git a/Engine/Plugins/Media/PixelStreaming2/Source/PixelStreaming2/Private/Blueprints/PixelStreaming2VideoComponent.h b/Engine/Plugins/Media/PixelStreaming2/Source/PixelStreaming2/Private/Blueprints/PixelStreaming2VideoComponent.h
new file mode 100644
--- /dev/null
+++ b/Engine/Plugins/Media/PixelStreaming2/Source/PixelStreaming2/Private/Blueprints/PixelStreaming2VideoComponent.h
@@ -0,0 +1,84 @@
+// Copyright Epic Games, Inc. All Rights Reserved.
+
+#pragma once
+
+#include "Components/SceneComponent.h"
+#include "CoreMinimal.h"
+#include "IPixelStreaming2VideoConsumer.h"
+#include "IPixelStreaming2VideoSink.h"
+#include "PixelStreaming2MediaTexture.h"
+
+#include "PixelStreaming2VideoComponent.generated.h"
+
+/**
+ * Allows in-engine playback of incoming WebRTC video from a particular Pixel Streaming player/peer using their camera in the browser.
+ * Note: Each video component associates itself with a particular Pixel Streaming player/peer (using the the Pixel Streaming player id).
+ */
+UCLASS(Blueprintable, ClassGroup = (PixelStreaming2), meta = (BlueprintSpawnableComponent))
+class PIXELSTREAMING2_API UPixelStreaming2VideoComponent : public USceneComponent
+{
+	GENERATED_BODY()
+
+protected:
+	UPixelStreaming2VideoComponent(const FObjectInitializer& ObjectInitializer);
+
+	//~ Begin UObject interface
+	virtual void BeginDestroy() override;
+	//~ End UObject interface
+
+	//~ Begin UActorComponent interface
+	virtual void TickComponent(float DeltaTime, ELevelTick TickType, FActorComponentTickFunction* ThisTickFunction) override;
+	//~ End UActorComponent interface
+
+public:
+	/**
+	 *   The Pixel Streaming streamer of the player that we wish to watch.
+	 *   If this is left blank this component will use the default Pixel Streaming streamer
+	 */
+	UPROPERTY(EditAnywhere, BlueprintReadWrite, Category = "Pixel Streaming Video Component")
+	FString StreamerToWatch;
+
+	/**
+	 *   The Pixel Streaming player/peer whose video we wish to watch.
+	 *   If this is left blank this component will watch the first non-watched to peer that connects after this component is ready.
+	 *   Note: that when the watched peer disconnects this component is reset to blank and will once again watch the next non-watched to peer that connects.
+	 */
+	UPROPERTY(EditAnywhere, BlueprintReadWrite, Category = "Pixel Streaming Video Component")
+	FString PlayerToWatch;
+
+	/**
+	 *  If not already watching a player/peer will try to attach for watching the "PlayerToWatch" each tick.
+	 */
+	UPROPERTY(EditAnywhere, BlueprintReadWrite, Category = "Pixel Streaming Video Component")
+	bool bAutoFindPeer;
+
+	/**
+	 * A sink for the video data received once this connection has finished negotiating.
+	 */
+	UPROPERTY(EditAnywhere, Category = "Pixel Streaming Video Component", META = (DisplayName = "Pixel Streaming Video Consumer", AllowPrivateAccess = true))
+	TObjectPtr<UPixelStreaming2MediaTexture> VideoConsumer = nullptr;
+
+private:
+	IPixelStreaming2VideoSink* VideoSink;
+
+	bool bIsWatchingPlayer;
+
+public:
+	// Watch a specific player on the default streamer. If the player is not found this component won't produce video.
+	UFUNCTION(BlueprintCallable, Category = "Pixel Streaming Video Component")
+	bool Watch(FString PlayerId);
+
+	// Watch a specific player. If the player is not found this component won't produce video.
+	UFUNCTION(BlueprintCallable, Category = "Pixel Streaming Video Component")
+	bool StreamerWatch(FString StreamerId, FString PlayerId);
+
+	// True if watching a connected WebRTC peer through Pixel Streaming.
+	UFUNCTION(BlueprintCallable, Category = "Pixel Streaming Video Component")
+	bool IsWatchingPlayer();
+
+	bool WillWatchAnyPlayer();
+
+	// Stops watching any connected player/peer and resets internal state so component is ready to watch again.
+	UFUNCTION(BlueprintCallable, Category = "Pixel Streaming Video Component")
+	void Reset();
+};
\ No newline at end of file
diff --git a/Engine/Plugins/Media/PixelStreaming2/Source/PixelStreaming2/Private/PixelStreaming2AudioComponent.cpp b/Engine/Plugins/Media/PixelStreaming2/Source/PixelStreaming2/Private/PixelStreaming2AudioComponent.cpp
deleted file mode 100644
diff --git a/Engine/Plugins/Media/PixelStreaming2/Source/PixelStreaming2/Private/PixelStreaming2AudioComponent.h b/Engine/Plugins/Media/PixelStreaming2/Source/PixelStreaming2/Private/PixelStreaming2AudioComponent.h
deleted file mode 100644
diff --git a/Engine/Plugins/Media/PixelStreaming2/Source/PixelStreaming2/Private/PixelStreaming2Blueprints.cpp b/Engine/Plugins/Media/PixelStreaming2/Source/PixelStreaming2/Private/PixelStreaming2Blueprints.cpp
deleted file mode 100644
diff --git a/Engine/Plugins/Media/PixelStreaming2/Source/PixelStreaming2/Private/PixelStreaming2Blueprints.h b/Engine/Plugins/Media/PixelStreaming2/Source/PixelStreaming2/Private/PixelStreaming2Blueprints.h
deleted file mode 100644
diff --git a/Engine/Plugins/Media/PixelStreaming2/Source/PixelStreaming2/Private/PixelStreaming2InputComponent.cpp b/Engine/Plugins/Media/PixelStreaming2/Source/PixelStreaming2/Private/PixelStreaming2InputComponent.cpp
deleted file mode 100644
diff --git a/Engine/Plugins/Media/PixelStreaming2/Source/PixelStreaming2/Private/PixelStreaming2InputComponent.h b/Engine/Plugins/Media/PixelStreaming2/Source/PixelStreaming2/Private/PixelStreaming2InputComponent.h
deleted file mode 100644
diff --git a/Engine/Plugins/Media/PixelStreaming2/Source/PixelStreaming2/Private/PixelStreaming2MediaTexture.cpp b/Engine/Plugins/Media/PixelStreaming2/Source/PixelStreaming2/Private/PixelStreaming2MediaTexture.cpp
deleted file mode 100644
diff --git a/Engine/Plugins/Media/PixelStreaming2/Source/PixelStreaming2/Private/PixelStreaming2MediaTextureResource.cpp b/Engine/Plugins/Media/PixelStreaming2/Source/PixelStreaming2/Private/PixelStreaming2MediaTextureResource.cpp
deleted file mode 100644
diff --git a/Engine/Plugins/Media/PixelStreaming2/Source/PixelStreaming2/Private/PixelStreaming2MediaTextureResource.h b/Engine/Plugins/Media/PixelStreaming2/Source/PixelStreaming2/Private/PixelStreaming2MediaTextureResource.h
deleted file mode 100644
diff --git a/Engine/Plugins/Media/PixelStreaming2/Source/PixelStreaming2/Private/PixelStreaming2Module.cpp b/Engine/Plugins/Media/PixelStreaming2/Source/PixelStreaming2/Private/PixelStreaming2Module.cpp
--- a/Engine/Plugins/Media/PixelStreaming2/Source/PixelStreaming2/Private/PixelStreaming2Module.cpp
+++ b/Engine/Plugins/Media/PixelStreaming2/Source/PixelStreaming2/Private/PixelStreaming2Module.cpp
@@ -2,6 +2,7 @@
 
 #include "PixelStreaming2Module.h"
 
+#include "Blueprints/PixelStreaming2InputComponent.h"
 #include "CoderUtils.h"
 #include "CoreMinimal.h"
 #include "CoreUtils.h"
@@ -12,7 +13,6 @@
 #include "Logging.h"
 #include "Modules/ModuleManager.h"
 #include "PixelStreaming2Delegates.h"
-#include "PixelStreaming2InputComponent.h"
 #include "PixelStreaming2Utils.h"
 #include "PixelStreaming2PluginSettings.h"
 #include "Slate/SceneViewport.h"
@@ -91,7 +91,7 @@ namespace UE::PixelStreaming2
 	static_assert(EpicRtcLogLevel::Debug == UnrealLogToEpicRtcCategoryMap[ELogVerbosity::Type::Verbose]);
 	static_assert(EpicRtcLogLevel::Trace == UnrealLogToEpicRtcCategoryMap[ELogVerbosity::Type::VeryVerbose]);
 	static_assert(EpicRtcLogLevel::Trace == UnrealLogToEpicRtcCategoryMap[ELogVerbosity::Type::All]);
-	
+
 	FPixelStreaming2Module* FPixelStreaming2Module::PixelStreaming2Module = nullptr;
 
 	FUtf8String FPixelStreaming2Module::EpicRtcConferenceName("pixel_streaming_conference_instance");
@@ -195,7 +195,7 @@ namespace UE::PixelStreaming2
 		if (UPixelStreaming2PluginSettings::FDelegates* Delegates = UPixelStreaming2PluginSettings::Delegates())
 		{
 			Delegates->OnLogStatsChanged.AddLambda([this](IConsoleVariable* Var) {
-				bool							 bLogStats = Var->GetBool();
+				bool					   bLogStats = Var->GetBool();
 				UPixelStreaming2Delegates* Delegates = UPixelStreaming2Delegates::Get();
 				if (!Delegates)
 				{
@@ -363,7 +363,7 @@ namespace UE::PixelStreaming2
 	TSharedPtr<IPixelStreaming2Streamer> FPixelStreaming2Module::DeleteStreamer(const FString& StreamerId)
 	{
 		TSharedPtr<IPixelStreaming2Streamer> ToBeDeleted;
-		FScopeLock								   Lock(&StreamersCS);
+		FScopeLock							 Lock(&StreamersCS);
 		if (Streamers.Contains(StreamerId))
 		{
 			ToBeDeleted = Streamers[StreamerId].Pin();
@@ -451,7 +451,7 @@ namespace UE::PixelStreaming2
 			// default to the scene viewport if we have a game engine
 			if (UGameEngine* GameEngine = Cast<UGameEngine>(GEngine))
 			{
-				TSharedPtr<SWindow>							   TargetWindow = GameEngine->GameViewport->GetWindow();
+				TSharedPtr<SWindow>						 TargetWindow = GameEngine->GameViewport->GetWindow();
 				TSharedPtr<IPixelStreaming2InputHandler> InputHandler = DefaultStreamer->GetInputHandler().Pin();
 				if (TargetWindow.IsValid() && InputHandler.IsValid())
 				{
diff --git a/Engine/Plugins/Media/PixelStreaming2/Source/PixelStreaming2/Private/PixelStreaming2Peer.cpp b/Engine/Plugins/Media/PixelStreaming2/Source/PixelStreaming2/Private/PixelStreaming2Peer.cpp
deleted file mode 100644
diff --git a/Engine/Plugins/Media/PixelStreaming2/Source/PixelStreaming2/Private/PixelStreaming2Peer.h b/Engine/Plugins/Media/PixelStreaming2/Source/PixelStreaming2/Private/PixelStreaming2Peer.h
deleted file mode 100644
diff --git a/Engine/Plugins/Media/PixelStreaming2/Source/PixelStreaming2/Private/PixelStreaming2VideoComponent.cpp b/Engine/Plugins/Media/PixelStreaming2/Source/PixelStreaming2/Private/PixelStreaming2VideoComponent.cpp
deleted file mode 100644
diff --git a/Engine/Plugins/Media/PixelStreaming2/Source/PixelStreaming2/Private/PixelStreaming2VideoComponent.h b/Engine/Plugins/Media/PixelStreaming2/Source/PixelStreaming2/Private/PixelStreaming2VideoComponent.h
deleted file mode 100644
