[PS2] QOL: Remove Blueprint and BlueprintEditor module by moving the respective files to the PixelStreaming2 and PixelStreaming2Editor module

diff --git a/Engine/Plugins/Media/PixelStreaming2/PixelStreaming2.uplugin b/Engine/Plugins/Media/PixelStreaming2/PixelStreaming2.uplugin
--- a/Engine/Plugins/Media/PixelStreaming2/PixelStreaming2.uplugin
+++ b/Engine/Plugins/Media/PixelStreaming2/PixelStreaming2.uplugin
@@ -43,20 +43,6 @@
             "PlatformAllowList": [ "Win64", "Linux" ],
             "TargetDenyList": [ "Server" ]
 		},
-		{
-			"Name": "PixelStreaming2Blueprint",
-			"Type": "Runtime",
-			"LoadingPhase": "PreDefault",
-			"PlatformAllowList": [ "Win64", "Linux" ],
-			"TargetDenyList": [ "Server" ]
-		},
-		{
-			"Name": "PixelStreaming2BlueprintEditor",
-			"Type": "Editor",
-			"LoadingPhase": "PostEngineInit",
-			"PlatformAllowList": [ "Win64", "Linux" ],
-			"TargetDenyList": [ "Server" ]
-		},
 		{
 			"Name": "PixelStreaming2HMD",
 			"Type": "Runtime",
diff --git a/Engine/Plugins/Media/PixelStreaming2/Source/PixelStreaming2/Internal/Blueprints/PixelStreaming2MediaTexture.h b/Engine/Plugins/Media/PixelStreaming2/Source/PixelStreaming2/Internal/Blueprints/PixelStreaming2MediaTexture.h
new file mode 100644
--- /dev/null
+++ b/Engine/Plugins/Media/PixelStreaming2/Source/PixelStreaming2/Internal/Blueprints/PixelStreaming2MediaTexture.h
@@ -0,0 +1,46 @@
+// Copyright Epic Games, Inc. All Rights Reserved.
+
+#pragma once
+
+#include "Engine/Texture2DDynamic.h"
+#include "IPixelStreaming2VideoConsumer.h"
+#include "RenderTargetPool.h"
+
+#include "PixelStreaming2MediaTexture.generated.h"
+
+namespace UE::PixelStreaming2
+{
+	class FPixelStreaming2MediaTextureResource;
+} // namespace UE::PixelStreaming2
+
+/**
+ * A Texture Object that can be used in materials etc. that takes updates from webrtc frames.
+ */
+UCLASS(NotBlueprintType, NotBlueprintable, HideDropdown, HideCategories = (ImportSettings, Compression, Texture, Adjustments, Compositing, LevelOfDetail, Object), META = (DisplayName = "PixelStreaming Media Texture"))
+class PIXELSTREAMING2_API UPixelStreaming2MediaTexture : public UTexture2DDynamic, public IPixelStreaming2VideoConsumer
+{
+	GENERATED_UCLASS_BODY()
+
+protected:
+	// UObject overrides.
+	virtual void BeginDestroy() override;
+	virtual void GetResourceSizeEx(FResourceSizeEx& CumulativeResourceSize) override;
+
+	// UTexture implementation
+	virtual FTextureResource* CreateResource() override;
+
+	// IPixelStreamingVideoConsumer implementation
+	virtual void ConsumeFrame(FTextureRHIRef Frame) override;
+	virtual void OnConsumerAdded() override {}
+	virtual void OnConsumerRemoved() override {}
+
+private:
+	void InitializeResources();
+
+	// updates the internal texture resource after each frame.
+	void UpdateTextureReference(FRHICommandList& RHICmdList, FTextureRHIRef Reference);
+
+	FCriticalSection RenderSyncContext;
+	// NOTE: This type has to be a raw ptr because UE will manage the lifetime of it
+	UE::PixelStreaming2::FPixelStreaming2MediaTextureResource* CurrentResource;
+};
diff --git a/Engine/Plugins/Media/PixelStreaming2/Source/PixelStreaming2/Internal/Blueprints/PixelStreaming2VideoProducer.h b/Engine/Plugins/Media/PixelStreaming2/Source/PixelStreaming2/Internal/Blueprints/PixelStreaming2VideoProducer.h
new file mode 100644
--- /dev/null
+++ b/Engine/Plugins/Media/PixelStreaming2/Source/PixelStreaming2/Internal/Blueprints/PixelStreaming2VideoProducer.h
@@ -0,0 +1,63 @@
+// Copyright Epic Games, Inc. All Rights Reserved.
+
+#pragma once
+
+#include "Engine/Texture2DDynamic.h"
+#include "Engine/TextureRenderTarget2D.h"
+#include "IPixelStreaming2VideoProducer.h"
+
+#include "PixelStreaming2VideoProducer.generated.h"
+
+/*
+ * The base video producer blueprint class
+ */
+UCLASS(NotBlueprintType, NotBlueprintable, Category = "PixelStreaming2", META = (DisplayName = "Video Producer"))
+class PIXELSTREAMING2_API UPixelStreaming2VideoProducer : public UObject
+{
+	GENERATED_UCLASS_BODY()
+
+public:
+	virtual TSharedPtr<IPixelStreaming2VideoProducer> GetVideoProducer() { return VideoProducer; }
+
+protected:
+	TSharedPtr<IPixelStreaming2VideoProducer> VideoProducer;
+};
+
+/*
+ * A video producer that streams the Unreal Engine's back buffer
+ */
+UCLASS(NotBlueprintType, NotBlueprintable, Category = "PixelStreaming2", META = (DisplayName = "Back Buffer Video Producer"))
+class PIXELSTREAMING2_API UPixelStreaming2VideoProducerBackBuffer : public UPixelStreaming2VideoProducer
+{
+	GENERATED_UCLASS_BODY()
+
+public:
+	virtual TSharedPtr<IPixelStreaming2VideoProducer> GetVideoProducer() override;
+};
+
+/*
+ * A video producer that streams the Unreal Engine's back buffer but uses the MediaIO Framework to capture the frame rather than Pixel Capture
+ */
+UCLASS(NotBlueprintType, NotBlueprintable, Category = "PixelStreaming2", META = (DisplayName = "Media Capture Video Producer"))
+class PIXELSTREAMING2_API UPixelStreaming2VideoProducerMediaCapture : public UPixelStreaming2VideoProducer
+{
+	GENERATED_UCLASS_BODY()
+
+public:
+	virtual TSharedPtr<IPixelStreaming2VideoProducer> GetVideoProducer() override;
+};
+
+/*
+ * A video producer that streams the specified render target 
+ */
+UCLASS(NotBlueprintType, NotBlueprintable, Category = "PixelStreaming2", META = (DisplayName = "Render Target Video Producer"))
+class PIXELSTREAMING2_API UPixelStreaming2VideoProducerRenderTarget : public UPixelStreaming2VideoProducer
+{
+	GENERATED_UCLASS_BODY()
+
+public:
+	UPROPERTY(BlueprintReadWrite, EditAnywhere, Category = Video, AssetRegistrySearchable)
+	TObjectPtr<UTextureRenderTarget2D> Target;
+
+	virtual TSharedPtr<IPixelStreaming2VideoProducer> GetVideoProducer() override;
+};
\ No newline at end of file
diff --git a/Engine/Plugins/Media/PixelStreaming2/Source/PixelStreaming2/Internal/PixelStreaming2MediaTexture.h b/Engine/Plugins/Media/PixelStreaming2/Source/PixelStreaming2/Internal/PixelStreaming2MediaTexture.h
deleted file mode 100644
diff --git a/Engine/Plugins/Media/PixelStreaming2/Source/PixelStreaming2/Private/Blueprints/PixelStreaming2AudioComponent.h b/Engine/Plugins/Media/PixelStreaming2/Source/PixelStreaming2/Private/Blueprints/PixelStreaming2AudioComponent.h
--- a/Engine/Plugins/Media/PixelStreaming2/Source/PixelStreaming2/Private/Blueprints/PixelStreaming2AudioComponent.h
+++ b/Engine/Plugins/Media/PixelStreaming2/Source/PixelStreaming2/Private/Blueprints/PixelStreaming2AudioComponent.h
@@ -18,7 +18,7 @@ namespace UE::PixelStreaming2
  * Note: Each audio component associates itself with a particular Pixel Streaming player/peer (using the the Pixel Streaming player id).
  */
 UCLASS(Blueprintable, ClassGroup = (PixelStreaming2), meta = (BlueprintSpawnableComponent))
-class PIXELSTREAMING2_API UPixelStreaming2AudioComponent : public USynthComponent, public IPixelStreaming2AudioConsumer
+class UPixelStreaming2AudioComponent : public USynthComponent, public IPixelStreaming2AudioConsumer
 {
 	GENERATED_BODY()
 
diff --git a/Engine/Plugins/Media/PixelStreaming2/Source/PixelStreaming2/Private/Blueprints/PixelStreaming2Blueprints.h b/Engine/Plugins/Media/PixelStreaming2/Source/PixelStreaming2/Private/Blueprints/PixelStreaming2Blueprints.h
--- a/Engine/Plugins/Media/PixelStreaming2/Source/PixelStreaming2/Private/Blueprints/PixelStreaming2Blueprints.h
+++ b/Engine/Plugins/Media/PixelStreaming2/Source/PixelStreaming2/Private/Blueprints/PixelStreaming2Blueprints.h
@@ -11,7 +11,7 @@
 #include "PixelStreaming2Blueprints.generated.h"
 
 UCLASS()
-class PIXELSTREAMING2_API UPixelStreaming2Blueprints : public UBlueprintFunctionLibrary
+class UPixelStreaming2Blueprints : public UBlueprintFunctionLibrary
 {
 	GENERATED_BODY()
 public:
diff --git a/Engine/Plugins/Media/PixelStreaming2/Source/PixelStreaming2/Private/Blueprints/PixelStreaming2InputComponent.h b/Engine/Plugins/Media/PixelStreaming2/Source/PixelStreaming2/Private/Blueprints/PixelStreaming2InputComponent.h
--- a/Engine/Plugins/Media/PixelStreaming2/Source/PixelStreaming2/Private/Blueprints/PixelStreaming2InputComponent.h
+++ b/Engine/Plugins/Media/PixelStreaming2/Source/PixelStreaming2/Private/Blueprints/PixelStreaming2InputComponent.h
@@ -12,7 +12,7 @@
  * Responses back to the source of the UI interactions may also be sent.
  */
 UCLASS(Blueprintable, ClassGroup = (PixelStreaming2), meta = (BlueprintSpawnableComponent))
-class PIXELSTREAMING2_API UPixelStreaming2Input : public UActorComponent
+class UPixelStreaming2Input : public UActorComponent
 {
 	GENERATED_BODY()
 
diff --git a/Engine/Plugins/Media/PixelStreaming2/Source/PixelStreaming2/Private/Blueprints/PixelStreaming2MediaTexture.cpp b/Engine/Plugins/Media/PixelStreaming2/Source/PixelStreaming2/Private/Blueprints/PixelStreaming2MediaTexture.cpp
--- a/Engine/Plugins/Media/PixelStreaming2/Source/PixelStreaming2/Private/Blueprints/PixelStreaming2MediaTexture.cpp
+++ b/Engine/Plugins/Media/PixelStreaming2/Source/PixelStreaming2/Private/Blueprints/PixelStreaming2MediaTexture.cpp
@@ -1,6 +1,6 @@
 // Copyright Epic Games, Inc. All Rights Reserved.
 
-#include "PixelStreaming2MediaTexture.h"
+#include "Blueprints/PixelStreaming2MediaTexture.h"
 
 #include "Async/Async.h"
 #include "DeviceProfiles/DeviceProfile.h"
diff --git a/Engine/Plugins/Media/PixelStreaming2/Source/PixelStreaming2/Private/Blueprints/PixelStreaming2Peer.h b/Engine/Plugins/Media/PixelStreaming2/Source/PixelStreaming2/Private/Blueprints/PixelStreaming2Peer.h
--- a/Engine/Plugins/Media/PixelStreaming2/Source/PixelStreaming2/Private/Blueprints/PixelStreaming2Peer.h
+++ b/Engine/Plugins/Media/PixelStreaming2/Source/PixelStreaming2/Private/Blueprints/PixelStreaming2Peer.h
@@ -2,12 +2,12 @@
 
 #pragma once
 
+#include "Blueprints/PixelStreaming2MediaTexture.h"
 #include "Components/SynthComponent.h"
 #include "Containers/Utf8String.h"
 #include "IPixelStreaming2AudioConsumer.h"
 #include "AudioSink.h"
 #include "EpicRtcManager.h"
-#include "PixelStreaming2MediaTexture.h"
 #include "VideoSink.h"
 
 #include "PixelStreaming2Peer.generated.h"
diff --git a/Engine/Plugins/Media/PixelStreaming2/Source/PixelStreaming2/Private/Blueprints/PixelStreaming2StreamerComponent.cpp b/Engine/Plugins/Media/PixelStreaming2/Source/PixelStreaming2/Private/Blueprints/PixelStreaming2StreamerComponent.cpp
new file mode 100644
--- /dev/null
+++ b/Engine/Plugins/Media/PixelStreaming2/Source/PixelStreaming2/Private/Blueprints/PixelStreaming2StreamerComponent.cpp
@@ -0,0 +1,157 @@
+// Copyright Epic Games, Inc. All Rights Reserved.
+
+#include "PixelStreaming2StreamerComponent.h"
+#include "IPixelStreaming2Module.h"
+#include "VideoProducerMediaCapture.h"
+#include "Engine/GameEngine.h"
+#include "Slate/SceneViewport.h"
+
+UPixelStreaming2StreamerComponent::UPixelStreaming2StreamerComponent(const FObjectInitializer& ObjectInitializer)
+{
+}
+
+void UPixelStreaming2StreamerComponent::BeginPlay() 
+{
+	if (Streamer)
+	{
+		SetupStreamerInput();
+	}
+}
+
+void UPixelStreaming2StreamerComponent::EndPlay(const EEndPlayReason::Type EndPlayReason)
+{
+	if (Streamer)
+	{
+		IPixelStreaming2Module::Get().DeleteStreamer(Streamer);
+	}
+}
+
+FString UPixelStreaming2StreamerComponent::GetId()
+{
+	if (Streamer)
+	{
+		return Streamer->GetId();
+	}
+	return "";
+}
+
+void UPixelStreaming2StreamerComponent::StartStreaming()
+{
+	if (!Streamer)
+	{
+		CreateStreamer();
+	}
+
+	if (VideoProducer && VideoProducer->GetVideoProducer())
+	{
+		Streamer->SetVideoProducer(VideoProducer->GetVideoProducer());
+	}
+	else
+	{
+		Streamer->SetVideoProducer(UE::PixelStreaming2::FVideoProducerMediaCapture::CreateActiveViewportCapture());
+	}
+
+	Streamer->SetCoupleFramerate(CoupleFramerate);
+	Streamer->SetStreamFPS(StreamFPS);
+	if (UsePixelStreamingURL)
+	{
+		FString ServerURL;
+		FParse::Value(FCommandLine::Get(), TEXT("PixelStreaming2URL="), ServerURL);
+		Streamer->SetSignallingServerURL(ServerURL);
+	}
+	else
+	{
+		Streamer->SetSignallingServerURL(SignallingServerURL);
+	}
+	Streamer->StartStreaming();
+}
+
+void UPixelStreaming2StreamerComponent::StopStreaming()
+{
+	if (Streamer)
+	{
+		Streamer->StopStreaming();
+	}
+}
+
+bool UPixelStreaming2StreamerComponent::IsStreaming()
+{
+	return Streamer && Streamer->IsStreaming();
+}
+
+void UPixelStreaming2StreamerComponent::ForceKeyFrame()
+{
+	if (Streamer)
+	{
+		Streamer->ForceKeyFrame();
+	}
+}
+
+void UPixelStreaming2StreamerComponent::FreezeStream(UTexture2D* Texture)
+{
+	if (Streamer)
+	{
+		Streamer->FreezeStream(Texture);
+	}
+}
+
+void UPixelStreaming2StreamerComponent::UnfreezeStream()
+{
+	if (Streamer)
+	{
+		Streamer->UnfreezeStream();
+	}
+}
+
+void UPixelStreaming2StreamerComponent::SendAllPlayersMessage(FString MessageType, const FString& Descriptor)
+{
+	if (Streamer)
+	{
+		Streamer->SendAllPlayersMessage(MessageType, Descriptor);
+	}
+}
+
+void UPixelStreaming2StreamerComponent::SendPlayerMessage(FString PlayerId, FString MessageType, const FString& Descriptor)
+{
+	if (Streamer)
+	{
+		Streamer->SendPlayerMessage(PlayerId, MessageType, Descriptor);
+	}
+}
+
+void UPixelStreaming2StreamerComponent::CreateStreamer()
+{
+	Streamer = IPixelStreaming2Module::Get().CreateStreamer(StreamerId);
+	Streamer->OnStreamingStarted().AddUObject(this, &UPixelStreaming2StreamerComponent::StreamingStarted);
+	Streamer->OnStreamingStopped().AddUObject(this, &UPixelStreaming2StreamerComponent::StreamingStopped);
+
+	SetupStreamerInput();
+}
+
+void UPixelStreaming2StreamerComponent::StreamingStarted(IPixelStreaming2Streamer*)
+{
+	OnStreamingStarted.Broadcast();
+}
+
+void UPixelStreaming2StreamerComponent::StreamingStopped(IPixelStreaming2Streamer*)
+{
+	OnStreamingStopped.Broadcast();
+}
+
+void UPixelStreaming2StreamerComponent::SetupStreamerInput()
+{
+	if (!GIsEditor)
+	{
+		// default to the scene viewport if we have a game engine
+		if (UGameEngine* GameEngine = Cast<UGameEngine>(GEngine))
+		{
+			TSharedPtr<FSceneViewport>					   TargetViewport = GameEngine->SceneViewport;
+			TSharedPtr<IPixelStreaming2InputHandler> InputHandler = Streamer->GetInputHandler().Pin();
+			if (TargetViewport.IsValid() && InputHandler.IsValid())
+			{
+				InputHandler->SetTargetViewport(TargetViewport->GetViewportWidget());
+				InputHandler->SetTargetWindow(TargetViewport->FindWindow());
+			}
+		}
+	}
+}
diff --git a/Engine/Plugins/Media/PixelStreaming2/Source/PixelStreaming2/Private/Blueprints/PixelStreaming2StreamerComponent.h b/Engine/Plugins/Media/PixelStreaming2/Source/PixelStreaming2/Private/Blueprints/PixelStreaming2StreamerComponent.h
new file mode 100644
--- /dev/null
+++ b/Engine/Plugins/Media/PixelStreaming2/Source/PixelStreaming2/Private/Blueprints/PixelStreaming2StreamerComponent.h
@@ -0,0 +1,82 @@
+// Copyright Epic Games, Inc. All Rights Reserved.
+
+#pragma once
+
+#include "Blueprints/PixelStreaming2VideoProducer.h"
+#include "IPixelStreaming2Streamer.h"
+#include "Components/ActorComponent.h"
+
+#include "PixelStreaming2StreamerComponent.generated.h"
+
+UCLASS(BlueprintType, Blueprintable, Category = "PixelStreaming2", META = (DisplayName = "Streamer Component", BlueprintSpawnableComponent))
+class UPixelStreaming2StreamerComponent : public UActorComponent
+{
+	GENERATED_UCLASS_BODY()
+
+public:
+	virtual void BeginPlay() override;
+	virtual void EndPlay(const EEndPlayReason::Type EndPlayReason) override;
+
+	UFUNCTION(BlueprintCallable, Category = "PixelStreaming2")
+	FString GetId();
+
+	UFUNCTION(BlueprintCallable, Category = "PixelStreaming2")
+	void StartStreaming();
+
+	UFUNCTION(BlueprintCallable, Category = "PixelStreaming2")
+	void StopStreaming();
+
+	UFUNCTION(BlueprintCallable, Category = "PixelStreaming2")
+	bool IsStreaming();
+
+	DECLARE_EVENT(UPixelStreaming2StreamerComponent, FStreamingStartedEvent);
+	FStreamingStartedEvent OnStreamingStarted;
+
+	DECLARE_EVENT(UPixelStreaming2StreamerComponent, FStreamingStoppedEvent);
+	FStreamingStoppedEvent OnStreamingStopped;
+
+	DECLARE_MULTICAST_DELEGATE_ThreeParams(FOnInputReceived, FString /* Peer id */, uint8 /* Type */, TArray<uint8> /* Message */);
+	FOnInputReceived OnInputReceived;
+
+	UFUNCTION(BlueprintCallable, Category = "PixelStreaming2")
+	void ForceKeyFrame();
+
+	UFUNCTION(BlueprintCallable, Category = "PixelStreaming2")
+	void FreezeStream(UTexture2D* Texture);
+
+	UFUNCTION(BlueprintCallable, Category = "PixelStreaming2")
+	void UnfreezeStream();
+
+	UFUNCTION(BlueprintCallable, Category = "PixelStreaming2")
+	void SendAllPlayersMessage(FString MessageType, const FString& Descriptor);
+
+	UFUNCTION(BlueprintCallable, Category = "PixelStreaming2")
+	void SendPlayerMessage(FString PlayerId, FString MessageType, const FString& Descriptor);
+
+	UPROPERTY(EditAnywhere, Category = "PixelStreaming2")
+	FString StreamerId = "Streamer Component";
+
+	UPROPERTY(EditAnywhere, Category = "PixelStreaming2")
+	FString SignallingServerURL = "ws://localhost:8888";
+
+	UPROPERTY(EditAnywhere, Category = "PixelStreaming2")
+	bool UsePixelStreamingURL = false;
+
+	UPROPERTY(EditAnywhere, Category = "PixelStreaming2")
+	int32 StreamFPS = 60;
+
+	UPROPERTY(EditAnywhere, Category = "PixelStreaming2")
+	bool CoupleFramerate = false;
+
+	UPROPERTY(EditAnywhere, Category = "PixelStreaming2")
+	TObjectPtr<UPixelStreaming2VideoProducer> VideoProducer = nullptr;
+
+private:
+	TSharedPtr<IPixelStreaming2Streamer> Streamer;
+
+	void CreateStreamer();
+	void SetupStreamerInput();
+
+	void StreamingStarted(IPixelStreaming2Streamer*);
+	void StreamingStopped(IPixelStreaming2Streamer*);
+};
diff --git a/Engine/Plugins/Media/PixelStreaming2/Source/PixelStreaming2/Private/Blueprints/PixelStreaming2VideoComponent.h b/Engine/Plugins/Media/PixelStreaming2/Source/PixelStreaming2/Private/Blueprints/PixelStreaming2VideoComponent.h
--- a/Engine/Plugins/Media/PixelStreaming2/Source/PixelStreaming2/Private/Blueprints/PixelStreaming2VideoComponent.h
+++ b/Engine/Plugins/Media/PixelStreaming2/Source/PixelStreaming2/Private/Blueprints/PixelStreaming2VideoComponent.h
@@ -2,11 +2,11 @@
 
 #pragma once
 
+#include "Blueprints/PixelStreaming2MediaTexture.h"
 #include "Components/SceneComponent.h"
 #include "CoreMinimal.h"
 #include "IPixelStreaming2VideoConsumer.h"
 #include "IPixelStreaming2VideoSink.h"
-#include "PixelStreaming2MediaTexture.h"
 
 #include "PixelStreaming2VideoComponent.generated.h"
 
@@ -15,7 +15,7 @@
  * Note: Each video component associates itself with a particular Pixel Streaming player/peer (using the the Pixel Streaming player id).
  */
 UCLASS(Blueprintable, ClassGroup = (PixelStreaming2), meta = (BlueprintSpawnableComponent))
-class PIXELSTREAMING2_API UPixelStreaming2VideoComponent : public USceneComponent
+class UPixelStreaming2VideoComponent : public USceneComponent
 {
 	GENERATED_BODY()
 
diff --git a/Engine/Plugins/Media/PixelStreaming2/Source/PixelStreaming2/Private/Blueprints/PixelStreaming2VideoProducer.cpp b/Engine/Plugins/Media/PixelStreaming2/Source/PixelStreaming2/Private/Blueprints/PixelStreaming2VideoProducer.cpp
new file mode 100644
--- /dev/null
+++ b/Engine/Plugins/Media/PixelStreaming2/Source/PixelStreaming2/Private/Blueprints/PixelStreaming2VideoProducer.cpp
@@ -0,0 +1,90 @@
+// Copyright Epic Games, Inc. All Rights Reserved.
+
+#include "Blueprints/PixelStreaming2VideoProducer.h"
+
+#include "VideoProducerBackBuffer.h"
+#include "VideoProducerMediaCapture.h"
+#include "VideoProducerPIEViewport.h"
+#include "VideoProducerRenderTarget.h"
+
+/*
+ * ---------- UPixelStreaming2VideoProducer -------------------
+ */
+UPixelStreaming2VideoProducer::UPixelStreaming2VideoProducer(const FObjectInitializer& ObjectInitializer)
+	: Super(ObjectInitializer)
+{
+}
+
+/*
+ * ---------- UPixelStreaming2VideoProducerBackBuffer -------------------
+ */
+UPixelStreaming2VideoProducerBackBuffer::UPixelStreaming2VideoProducerBackBuffer(const FObjectInitializer& ObjectInitializer)
+	: Super(ObjectInitializer)
+{
+}
+
+TSharedPtr<IPixelStreaming2VideoProducer> UPixelStreaming2VideoProducerBackBuffer::GetVideoProducer()
+{
+	if (!VideoProducer)
+	{
+		// detect if we're in PIE mode or not
+		bool bIsGame = false;
+		FParse::Bool(FCommandLine::Get(), TEXT("game"), bIsGame);
+		if (GIsEditor && !bIsGame)
+		{
+			VideoProducer = UE::PixelStreaming2::FVideoProducerPIEViewport::Create();
+		}
+		else
+		{
+			VideoProducer = UE::PixelStreaming2::FVideoProducerBackBuffer::Create();
+		}
+	}
+
+	return VideoProducer;
+}
+
+/*
+ * ---------- UPixelStreaming2VideoProducerMediaCapture -------------------
+ */
+UPixelStreaming2VideoProducerMediaCapture::UPixelStreaming2VideoProducerMediaCapture(const FObjectInitializer& ObjectInitializer)
+	: Super(ObjectInitializer)
+{
+}
+
+TSharedPtr<IPixelStreaming2VideoProducer> UPixelStreaming2VideoProducerMediaCapture::GetVideoProducer()
+{
+	if (!VideoProducer)
+	{
+		// detect if we're in PIE mode or not
+		bool bIsGame = false;
+		FParse::Bool(FCommandLine::Get(), TEXT("game"), bIsGame);
+		if (GIsEditor && !bIsGame)
+		{
+			VideoProducer = UE::PixelStreaming2::FVideoProducerPIEViewport::Create();
+		}
+		else
+		{
+			VideoProducer = UE::PixelStreaming2::FVideoProducerMediaCapture::CreateActiveViewportCapture();
+		}
+	}
+
+	return VideoProducer;
+}
+
+/*
+ * ---------- UPixelStreaming2VideoProducerRenderTarget -------------------
+ */
+UPixelStreaming2VideoProducerRenderTarget::UPixelStreaming2VideoProducerRenderTarget(const FObjectInitializer& ObjectInitializer)
+	: Super(ObjectInitializer)
+{
+}
+
+TSharedPtr<IPixelStreaming2VideoProducer> UPixelStreaming2VideoProducerRenderTarget::GetVideoProducer()
+{
+	if (!VideoProducer)
+	{
+		VideoProducer = UE::PixelStreaming2::FVideoProducerRenderTarget::Create(Target);
+	}
+	return VideoProducer;
+}
+
diff --git a/Engine/Plugins/Media/PixelStreaming2/Source/PixelStreaming2Blueprint/Internal/PixelStreaming2StreamerVideoProducer.h b/Engine/Plugins/Media/PixelStreaming2/Source/PixelStreaming2Blueprint/Internal/PixelStreaming2StreamerVideoProducer.h
deleted file mode 100644
diff --git a/Engine/Plugins/Media/PixelStreaming2/Source/PixelStreaming2Blueprint/Internal/PixelStreaming2StreamerVideoProducerBackBuffer.h b/Engine/Plugins/Media/PixelStreaming2/Source/PixelStreaming2Blueprint/Internal/PixelStreaming2StreamerVideoProducerBackBuffer.h
deleted file mode 100644
diff --git a/Engine/Plugins/Media/PixelStreaming2/Source/PixelStreaming2Blueprint/Internal/PixelStreaming2StreamerVideoProducerMediaCapture.h b/Engine/Plugins/Media/PixelStreaming2/Source/PixelStreaming2Blueprint/Internal/PixelStreaming2StreamerVideoProducerMediaCapture.h
deleted file mode 100644
diff --git a/Engine/Plugins/Media/PixelStreaming2/Source/PixelStreaming2Blueprint/Internal/PixelStreaming2StreamerVideoProducerRenderTarget.h b/Engine/Plugins/Media/PixelStreaming2/Source/PixelStreaming2Blueprint/Internal/PixelStreaming2StreamerVideoProducerRenderTarget.h
deleted file mode 100644
diff --git a/Engine/Plugins/Media/PixelStreaming2/Source/PixelStreaming2Blueprint/PixelStreaming2Blueprint.Build.cs b/Engine/Plugins/Media/PixelStreaming2/Source/PixelStreaming2Blueprint/PixelStreaming2Blueprint.Build.cs
deleted file mode 100644
diff --git a/Engine/Plugins/Media/PixelStreaming2/Source/PixelStreaming2Blueprint/Private/PixelStreaming2BlueprintModule.cpp b/Engine/Plugins/Media/PixelStreaming2/Source/PixelStreaming2Blueprint/Private/PixelStreaming2BlueprintModule.cpp
deleted file mode 100644
diff --git a/Engine/Plugins/Media/PixelStreaming2/Source/PixelStreaming2Blueprint/Private/PixelStreaming2BlueprintPrivate.h b/Engine/Plugins/Media/PixelStreaming2/Source/PixelStreaming2Blueprint/Private/PixelStreaming2BlueprintPrivate.h
deleted file mode 100644
diff --git a/Engine/Plugins/Media/PixelStreaming2/Source/PixelStreaming2Blueprint/Private/PixelStreaming2StreamerComponent.cpp b/Engine/Plugins/Media/PixelStreaming2/Source/PixelStreaming2Blueprint/Private/PixelStreaming2StreamerComponent.cpp
deleted file mode 100644
diff --git a/Engine/Plugins/Media/PixelStreaming2/Source/PixelStreaming2Blueprint/Private/PixelStreaming2StreamerComponent.h b/Engine/Plugins/Media/PixelStreaming2/Source/PixelStreaming2Blueprint/Private/PixelStreaming2StreamerComponent.h
deleted file mode 100644
diff --git a/Engine/Plugins/Media/PixelStreaming2/Source/PixelStreaming2Blueprint/Private/PixelStreaming2StreamerVideoProducer.cpp b/Engine/Plugins/Media/PixelStreaming2/Source/PixelStreaming2Blueprint/Private/PixelStreaming2StreamerVideoProducer.cpp
deleted file mode 100644
diff --git a/Engine/Plugins/Media/PixelStreaming2/Source/PixelStreaming2Blueprint/Private/PixelStreaming2StreamerVideoProducerBackBuffer.cpp b/Engine/Plugins/Media/PixelStreaming2/Source/PixelStreaming2Blueprint/Private/PixelStreaming2StreamerVideoProducerBackBuffer.cpp
deleted file mode 100644
diff --git a/Engine/Plugins/Media/PixelStreaming2/Source/PixelStreaming2Blueprint/Private/PixelStreaming2StreamerVideoProducerMediaCapture.cpp b/Engine/Plugins/Media/PixelStreaming2/Source/PixelStreaming2Blueprint/Private/PixelStreaming2StreamerVideoProducerMediaCapture.cpp
deleted file mode 100644
diff --git a/Engine/Plugins/Media/PixelStreaming2/Source/PixelStreaming2Blueprint/Private/PixelStreaming2StreamerVideoProducerRenderTarget.cpp b/Engine/Plugins/Media/PixelStreaming2/Source/PixelStreaming2Blueprint/Private/PixelStreaming2StreamerVideoProducerRenderTarget.cpp
deleted file mode 100644
diff --git a/Engine/Plugins/Media/PixelStreaming2/Source/PixelStreaming2BlueprintEditor/PixelStreaming2BlueprintEditor.Build.cs b/Engine/Plugins/Media/PixelStreaming2/Source/PixelStreaming2BlueprintEditor/PixelStreaming2BlueprintEditor.Build.cs
deleted file mode 100644
diff --git a/Engine/Plugins/Media/PixelStreaming2/Source/PixelStreaming2BlueprintEditor/Private/AssetTypeActions_StreamerVideoProducer.h b/Engine/Plugins/Media/PixelStreaming2/Source/PixelStreaming2BlueprintEditor/Private/AssetTypeActions_StreamerVideoProducer.h
deleted file mode 100644
diff --git a/Engine/Plugins/Media/PixelStreaming2/Source/PixelStreaming2BlueprintEditor/Private/PixelStreaming2BlueprintEditorModule.cpp b/Engine/Plugins/Media/PixelStreaming2/Source/PixelStreaming2BlueprintEditor/Private/PixelStreaming2BlueprintEditorModule.cpp
deleted file mode 100644
diff --git a/Engine/Plugins/Media/PixelStreaming2/Source/PixelStreaming2BlueprintEditor/Private/PixelStreaming2MediaTextureFactory.cpp b/Engine/Plugins/Media/PixelStreaming2/Source/PixelStreaming2BlueprintEditor/Private/PixelStreaming2MediaTextureFactory.cpp
deleted file mode 100644
diff --git a/Engine/Plugins/Media/PixelStreaming2/Source/PixelStreaming2BlueprintEditor/Private/PixelStreaming2MediaTextureFactory.h b/Engine/Plugins/Media/PixelStreaming2/Source/PixelStreaming2BlueprintEditor/Private/PixelStreaming2MediaTextureFactory.h
deleted file mode 100644
diff --git a/Engine/Plugins/Media/PixelStreaming2/Source/PixelStreaming2BlueprintEditor/Private/PixelStreaming2StreamerVideoProducerBackBufferFactory.cpp b/Engine/Plugins/Media/PixelStreaming2/Source/PixelStreaming2BlueprintEditor/Private/PixelStreaming2StreamerVideoProducerBackBufferFactory.cpp
deleted file mode 100644
diff --git a/Engine/Plugins/Media/PixelStreaming2/Source/PixelStreaming2BlueprintEditor/Private/PixelStreaming2StreamerVideoProducerBackBufferFactory.h b/Engine/Plugins/Media/PixelStreaming2/Source/PixelStreaming2BlueprintEditor/Private/PixelStreaming2StreamerVideoProducerBackBufferFactory.h
deleted file mode 100644
diff --git a/Engine/Plugins/Media/PixelStreaming2/Source/PixelStreaming2BlueprintEditor/Private/PixelStreaming2StreamerVideoProducerMediaCaptureFactory.cpp b/Engine/Plugins/Media/PixelStreaming2/Source/PixelStreaming2BlueprintEditor/Private/PixelStreaming2StreamerVideoProducerMediaCaptureFactory.cpp
deleted file mode 100644
diff --git a/Engine/Plugins/Media/PixelStreaming2/Source/PixelStreaming2BlueprintEditor/Private/PixelStreaming2StreamerVideoProducerMediaCaptureFactory.h b/Engine/Plugins/Media/PixelStreaming2/Source/PixelStreaming2BlueprintEditor/Private/PixelStreaming2StreamerVideoProducerMediaCaptureFactory.h
deleted file mode 100644
diff --git a/Engine/Plugins/Media/PixelStreaming2/Source/PixelStreaming2BlueprintEditor/Private/PixelStreaming2StreamerVideoProducerRenderTargetFactory.cpp b/Engine/Plugins/Media/PixelStreaming2/Source/PixelStreaming2BlueprintEditor/Private/PixelStreaming2StreamerVideoProducerRenderTargetFactory.cpp
deleted file mode 100644
diff --git a/Engine/Plugins/Media/PixelStreaming2/Source/PixelStreaming2BlueprintEditor/Private/PixelStreaming2StreamerVideoProducerRenderTargetFactory.h b/Engine/Plugins/Media/PixelStreaming2/Source/PixelStreaming2BlueprintEditor/Private/PixelStreaming2StreamerVideoProducerRenderTargetFactory.h
deleted file mode 100644
diff --git a/Engine/Plugins/Media/PixelStreaming2/Source/PixelStreaming2Editor/Private/AssetTypeActions_VideoProducer.h b/Engine/Plugins/Media/PixelStreaming2/Source/PixelStreaming2Editor/Private/AssetTypeActions_VideoProducer.h
new file mode 100644
--- /dev/null
+++ b/Engine/Plugins/Media/PixelStreaming2/Source/PixelStreaming2Editor/Private/AssetTypeActions_VideoProducer.h
@@ -0,0 +1,25 @@
+// Copyright Epic Games, Inc. All Rights Reserved.
+
+#pragma once
+
+#include "AssetTypeActions_Base.h"
+#include "Blueprints/PixelStreaming2VideoProducer.h"
+#include "CoreMinimal.h"
+#include "Toolkits/IToolkitHost.h"
+
+class FAssetTypeActions_VideoProducer : public FAssetTypeActions_Base
+{
+public:
+	// IAssetTypeActions Implementation
+	virtual FText GetName() const override { return NSLOCTEXT("PixelStreaming2", "AssetTypeActions_VideoProducer", "Video Input Actions"); }
+	virtual FColor GetTypeColor() const override { return FColor(192,64,64); }
+	virtual UClass* GetSupportedClass() const override { return UPixelStreaming2VideoProducer::StaticClass(); }
+	virtual bool CanFilter() override { return false; }
+	virtual uint32 GetCategories() override
+	{
+		IAssetTools& AssetTools = FModuleManager::LoadModuleChecked<FAssetToolsModule>("AssetTools").Get();
+		return AssetTools.RegisterAdvancedAssetCategory("PixelStreaming2", NSLOCTEXT("PixelStreaming2", "AssetCategoryDisplayName", "PixelStreaming2"));
+		//return EAssetTypeCategories::Misc;
+	}
+	virtual bool IsImportedAsset() const override { return false; }
+};
diff --git a/Engine/Plugins/Media/PixelStreaming2/Source/PixelStreaming2Editor/Private/Factories.cpp b/Engine/Plugins/Media/PixelStreaming2/Source/PixelStreaming2Editor/Private/Factories.cpp
new file mode 100644
--- /dev/null
+++ b/Engine/Plugins/Media/PixelStreaming2/Source/PixelStreaming2Editor/Private/Factories.cpp
@@ -0,0 +1,144 @@
+// Copyright Epic Games, Inc. All Rights Reserved.
+
+#include "Factories.h"
+
+#include "AssetToolsModule.h"
+#include "AssetTypeCategories.h"
+#include "Blueprints/PixelStreaming2MediaTexture.h"
+#include "Blueprints/PixelStreaming2VideoProducer.h"
+#include "IAssetTools.h"
+#include "Modules/ModuleManager.h"
+
+#define LOCTEXT_NAMESPACE "PixelStreaming2"
+
+/*
+ * ---------- UPixelStreaming2MediaTextureFactory -------------------
+ */
+UPixelStreaming2MediaTextureFactory::UPixelStreaming2MediaTextureFactory(const FObjectInitializer& ObjectInitializer)
+	: Super(ObjectInitializer)
+{
+	bCreateNew = true;
+	bEditAfterNew = true;
+
+	SupportedClass = UPixelStreaming2MediaTexture::StaticClass();
+}
+
+FText UPixelStreaming2MediaTextureFactory::GetDisplayName() const
+{
+	return LOCTEXT("MediaTextureFactoryDisplayName", "Pixel Streaming 2 Media Texture");
+}
+
+uint32 UPixelStreaming2MediaTextureFactory::GetMenuCategories() const
+{
+	return EAssetTypeCategories::Textures;
+}
+
+UObject* UPixelStreaming2MediaTextureFactory::FactoryCreateNew(UClass* InClass, UObject* InParent, FName InName, EObjectFlags Flags, UObject* Context, FFeedbackContext* Warn)
+{
+	if (UPixelStreaming2MediaTexture* Resource = NewObject<UPixelStreaming2MediaTexture>(InParent, InName, Flags | RF_Transactional))
+	{
+		Resource->UpdateResource();
+		return Resource;
+	}
+
+	return nullptr;
+}
+
+/*
+ * ---------- UPixelStreaming2VideoProducerBackBufferFactory -------------------
+ */
+UPixelStreaming2VideoProducerBackBufferFactory::UPixelStreaming2VideoProducerBackBufferFactory(const FObjectInitializer& ObjectInitializer)
+	: Super(ObjectInitializer)
+{
+	bCreateNew = true;
+	bEditAfterNew = true;
+	SupportedClass = UPixelStreaming2VideoProducerBackBuffer::StaticClass();
+}
+
+FText UPixelStreaming2VideoProducerBackBufferFactory::GetDisplayName() const
+{
+	return LOCTEXT("VideoProducerBackBufferDisplayName", "Back Buffer  Video Input");
+}
+
+uint32 UPixelStreaming2VideoProducerBackBufferFactory::GetMenuCategories() const
+{
+	IAssetTools& AssetTools = FModuleManager::LoadModuleChecked<FAssetToolsModule>("AssetTools").Get();
+	return AssetTools.RegisterAdvancedAssetCategory("PixelStreaming2", LOCTEXT("AssetCategoryDisplayName", "PixelStreaming2"));
+}
+
+UObject* UPixelStreaming2VideoProducerBackBufferFactory::FactoryCreateNew(UClass* InClass, UObject* InParent, FName InName, EObjectFlags Flags, UObject* Context, FFeedbackContext* Warn)
+{
+	if (UPixelStreaming2VideoProducerBackBuffer* Resource = NewObject<UPixelStreaming2VideoProducerBackBuffer>(InParent, InName, Flags | RF_Transactional))
+	{
+		return Resource;
+	}
+
+	return nullptr;
+}
+
+/*
+ * ---------- UPixelStreaming2VideoProducerMediaCaptureFactory -------------------
+ */
+UPixelStreaming2VideoProducerMediaCaptureFactory::UPixelStreaming2VideoProducerMediaCaptureFactory(const FObjectInitializer& ObjectInitializer)
+	: Super(ObjectInitializer)
+{
+	bCreateNew = true;
+	bEditAfterNew = true;
+	SupportedClass = UPixelStreaming2VideoProducerMediaCapture::StaticClass();
+}
+
+FText UPixelStreaming2VideoProducerMediaCaptureFactory::GetDisplayName() const
+{
+	return LOCTEXT("VideoProducerMediaCaptureDisplayName", "Media Capture  Video Input");
+}
+
+uint32 UPixelStreaming2VideoProducerMediaCaptureFactory::GetMenuCategories() const
+{
+	IAssetTools& AssetTools = FModuleManager::LoadModuleChecked<FAssetToolsModule>("AssetTools").Get();
+	return AssetTools.RegisterAdvancedAssetCategory("PixelStreaming2", LOCTEXT("AssetCategoryDisplayName", "PixelStreaming2"));
+}
+
+UObject* UPixelStreaming2VideoProducerMediaCaptureFactory::FactoryCreateNew(UClass* InClass, UObject* InParent, FName InName, EObjectFlags Flags, UObject* Context, FFeedbackContext* Warn)
+{
+	if (UPixelStreaming2VideoProducerMediaCapture* Resource = NewObject<UPixelStreaming2VideoProducerMediaCapture>(InParent, InName, Flags | RF_Transactional))
+	{
+		return Resource;
+	}
+
+	return nullptr;
+}
+
+/*
+ * ---------- UPixelStreaming2VideoProducerRenderTargetFactory -------------------
+ */
+UPixelStreaming2VideoProducerRenderTargetFactory::UPixelStreaming2VideoProducerRenderTargetFactory(const FObjectInitializer& ObjectInitializer)
+	: Super(ObjectInitializer)
+{
+	bCreateNew = true;
+	bEditAfterNew = true;
+
+	SupportedClass = UPixelStreaming2VideoProducerRenderTarget::StaticClass();
+}
+
+FText UPixelStreaming2VideoProducerRenderTargetFactory::GetDisplayName() const
+{
+	return LOCTEXT("VideoProducerRenderTargetDisplayName", "Render Target  Video Input");
+}
+
+uint32 UPixelStreaming2VideoProducerRenderTargetFactory::GetMenuCategories() const
+{
+	IAssetTools& AssetTools = FModuleManager::LoadModuleChecked<FAssetToolsModule>("AssetTools").Get();
+	return AssetTools.RegisterAdvancedAssetCategory("PixelStreaming2", LOCTEXT("AssetCategoryDisplayName", "PixelStreaming2"));
+}
+
+UObject* UPixelStreaming2VideoProducerRenderTargetFactory::FactoryCreateNew(UClass* InClass, UObject* InParent, FName InName, EObjectFlags Flags, UObject* Context, FFeedbackContext* Warn)
+{
+	if (UPixelStreaming2VideoProducerRenderTarget* Resource = NewObject<UPixelStreaming2VideoProducerRenderTarget>(InParent, InName, Flags | RF_Transactional))
+	{
+		return Resource;
+	}
+
+	return nullptr;
+}
+
+#undef LOCTEXT_NAMESPACE
diff --git a/Engine/Plugins/Media/PixelStreaming2/Source/PixelStreaming2Editor/Private/Factories.h b/Engine/Plugins/Media/PixelStreaming2/Source/PixelStreaming2Editor/Private/Factories.h
new file mode 100644
--- /dev/null
+++ b/Engine/Plugins/Media/PixelStreaming2/Source/PixelStreaming2Editor/Private/Factories.h
@@ -0,0 +1,65 @@
+// Copyright Epic Games, Inc. All Rights Reserved.
+
+#pragma once
+
+#include "CoreMinimal.h"
+#include "Factories/Factory.h"
+#include "UObject/Object.h"
+
+#include "Factories.generated.h"
+
+/*
+ * ---------- UPixelStreaming2MediaTextureFactory -------------------
+ */
+UCLASS()
+class UPixelStreaming2MediaTextureFactory : public UFactory
+{
+	GENERATED_UCLASS_BODY()
+
+public:
+	virtual FText	 GetDisplayName() const override;
+	virtual uint32	 GetMenuCategories() const override;
+	virtual UObject* FactoryCreateNew(UClass* InClass, UObject* InParent, FName InName, EObjectFlags Flags, UObject* Context, FFeedbackContext* Warn) override;
+};
+
+/*
+ * ---------- UPixelStreaming2VideoProducerBackBufferFactory -------------------
+ */
+UCLASS()
+class UPixelStreaming2VideoProducerBackBufferFactory : public UFactory
+{
+	GENERATED_UCLASS_BODY()
+
+public:
+	virtual FText	 GetDisplayName() const override;
+	virtual uint32	 GetMenuCategories() const override;
+	virtual UObject* FactoryCreateNew(UClass* InClass, UObject* InParent, FName InName, EObjectFlags Flags, UObject* Context, FFeedbackContext* Warn) override;
+};
+
+/*
+ * ---------- UPixelStreaming2VideoProducerMediaCaptureFactory -------------------
+ */
+UCLASS()
+class UPixelStreaming2VideoProducerMediaCaptureFactory : public UFactory
+{
+	GENERATED_UCLASS_BODY()
+
+public:
+	virtual FText	 GetDisplayName() const override;
+	virtual uint32	 GetMenuCategories() const override;
+	virtual UObject* FactoryCreateNew(UClass* InClass, UObject* InParent, FName InName, EObjectFlags Flags, UObject* Context, FFeedbackContext* Warn) override;
+};
+
+/*
+ * ---------- UPixelStreaming2VideoProducerRenderTargetFactory -------------------
+ */
+UCLASS()
+class UPixelStreaming2VideoProducerRenderTargetFactory : public UFactory
+{
+	GENERATED_UCLASS_BODY()
+
+public:
+	virtual FText	 GetDisplayName() const override;
+	virtual uint32	 GetMenuCategories() const override;
+	virtual UObject* FactoryCreateNew(UClass* InClass, UObject* InParent, FName InName, EObjectFlags Flags, UObject* Context, FFeedbackContext* Warn) override;
+};
\ No newline at end of file
diff --git a/Engine/Plugins/Media/PixelStreaming2/Source/PixelStreaming2Editor/Private/PixelStreaming2EditorModule.cpp b/Engine/Plugins/Media/PixelStreaming2/Source/PixelStreaming2Editor/Private/PixelStreaming2EditorModule.cpp
--- a/Engine/Plugins/Media/PixelStreaming2/Source/PixelStreaming2Editor/Private/PixelStreaming2EditorModule.cpp
+++ b/Engine/Plugins/Media/PixelStreaming2/Source/PixelStreaming2Editor/Private/PixelStreaming2EditorModule.cpp
@@ -4,6 +4,7 @@
 
 #include "PixelStreaming2Style.h"
 #include "PixelStreaming2Toolbar.h"
+#include "AssetTypeActions_VideoProducer.h"
 
 #include "CoreUtils.h"
 #include "Editor/EditorPerformanceSettings.h"
@@ -31,6 +32,8 @@
 
 #define LOCTEXT_NAMESPACE "PixelStreaming2EditorModule"
 
+#define IMAGE_BRUSH_SVG(Style, RelativePath, ...) FSlateVectorImageBrush(Style.RootToContentDir(RelativePath, TEXT(".svg")), __VA_ARGS__)
+
 namespace UE::EditorPixelStreaming2
 {
 	/**
@@ -42,10 +45,26 @@ namespace UE::EditorPixelStreaming2
 		{
 			return;
 		}
-		
+
+		IAssetTools& AssetTools = FModuleManager::LoadModuleChecked<FAssetToolsModule>("AssetTools").Get();
+		AssetTools.RegisterAssetTypeActions(MakeShared<FAssetTypeActions_VideoProducer>());
+
 		// Initialize the editor toolbar
 		FPixelStreaming2Style::Initialize();
+		FSlateStyleSet& StyleInstance = UE::EditorPixelStreaming2::FPixelStreaming2Style::Get();
+
+		const FVector2D Icon20x20(20.0f, 20.0f);
+		const FVector2D Icon64x64(64.0f, 64.0f);
+
+		StyleInstance.Set("ClassThumbnail.PixelStreaming2VideoProducerBackBuffer", new IMAGE_BRUSH_SVG(StyleInstance, "PixelStreaming2_64", Icon64x64));
+		StyleInstance.Set("ClassIcon.PixelStreaming2VideoProducerBackBuffer", new IMAGE_BRUSH_SVG(StyleInstance, "PixelStreaming2_20", Icon20x20));
+		StyleInstance.Set("ClassThumbnail.PixelStreaming2VideoProducerRenderTarget", new IMAGE_BRUSH_SVG(StyleInstance, "PixelStreaming2_64", Icon64x64));
+		StyleInstance.Set("ClassIcon.PixelStreaming2VideoProducerRenderTarget", new IMAGE_BRUSH_SVG(StyleInstance, "PixelStreaming2_20", Icon20x20));
+		StyleInstance.Set("ClassThumbnail.PixelStreaming2VideoProducerMediaCapture", new IMAGE_BRUSH_SVG(StyleInstance, "PixelStreaming2_64", Icon64x64));
+		StyleInstance.Set("ClassIcon.PixelStreaming2VideoProducerMediaCapture", new IMAGE_BRUSH_SVG(StyleInstance, "PixelStreaming2_20", Icon20x20));
+
 		FPixelStreaming2Style::ReloadTextures();
+
 		Toolbar = MakeShared<FPixelStreaming2Toolbar>();
 
 		if (UPixelStreaming2Delegates* Delegates = UPixelStreaming2Delegates::Get())
@@ -421,6 +440,9 @@ namespace UE::EditorPixelStreaming2
 		}
 	}
 } // namespace UE::EditorPixelStreaming2
+
+#undef IMAGE_BRUSH_SVG
+
 #undef LOCTEXT_NAMESPACE
 
 IMPLEMENT_MODULE(UE::EditorPixelStreaming2::FPixelStreaming2EditorModule, PixelStreaming2Editor)
