Fix Pixel Streaming 2 not streaming h264 with AMD gpu.

Problem is WebRTC does not support using SendRecv if the encoding and decoding codec do not support the same codec. AMD GPUs currently have decoding disabled so WebRTC fails to negotiate codecs.

This fix checks for AMD GPU and disables SendRecv for h264.

See CL 36069032 for a more generic fix.

diff --git a/Engine/Plugins/Media/PixelStreaming2/Source/PixelStreaming2/Private/Streamer.cpp b/Engine/Plugins/Media/PixelStreaming2/Source/PixelStreaming2/Private/Streamer.cpp
--- a/Engine/Plugins/Media/PixelStreaming2/Source/PixelStreaming2/Private/Streamer.cpp
+++ b/Engine/Plugins/Media/PixelStreaming2/Source/PixelStreaming2/Private/Streamer.cpp
@@ -949,8 +949,23 @@ namespace UE::PixelStreaming2
 			ParticipantConnection->AddAudioSource(AudioSource);
 		}
 
+		const EVideoCodec SelectedCodec = GetEnumFromCVar<EVideoCodec>(UPixelStreaming2PluginSettings::CVarEncoderCodec);
+		const bool bNegotiateCodecs = UPixelStreaming2PluginSettings::CVarWebRTCNegotiateCodecs.GetValueOnAnyThread();
 		const bool bTransmitUEVideo = !UPixelStreaming2PluginSettings::CVarWebRTCDisableTransmitVideo.GetValueOnAnyThread();
-		const bool bReceiveBrowserVideo = !UPixelStreaming2PluginSettings::CVarWebRTCDisableReceiveVideo.GetValueOnAnyThread();
+		bool bReceiveBrowserVideo = !UPixelStreaming2PluginSettings::CVarWebRTCDisableReceiveVideo.GetValueOnAnyThread();
+		
+		// Check if the user has selected only H.264 on a AMD gpu and disable receiving video.
+		// WebRTC does not support using SendRecv if the encoding and decoding do not support the same codec. 
+		// AMD GPUs currently have decoding disabled so WebRTC fails to create SDP codecs with SendRecv.
+		if (IsRHIDeviceAMD() && !bNegotiateCodecs && SelectedCodec == EVideoCodec::H264)
+		{
+			if (bReceiveBrowserVideo)
+			{
+				bReceiveBrowserVideo = false;
+				UE_LOGFMT(LogPixelStreaming2, Warning, "AMD GPUs do not support receiving H.264 video.");
+			}
+		}
+
 		if (bTransmitUEVideo || bReceiveBrowserVideo)
 		{
 			TArray<EpicRtcVideoEncodingConfig> VideoEncodingConfigs;
