Fix Pixel Streaming 2 internal signalling server returning 404

Latest PixelStreaming Infra places the web files in a different folder. This fix changes to use this filder.

diff --git a/Engine/Plugins/Media/PixelStreaming/Source/PixelStreamingServers/Private/SignallingServer.cpp b/Engine/Plugins/Media/PixelStreaming/Source/PixelStreamingServers/Private/SignallingServer.cpp
--- a/Engine/Plugins/Media/PixelStreaming/Source/PixelStreamingServers/Private/SignallingServer.cpp
+++ b/Engine/Plugins/Media/PixelStreaming/Source/PixelStreamingServers/Private/SignallingServer.cpp
@@ -158,7 +158,7 @@ namespace UE::PixelStreamingServers
 
 		// Add /Public
 		FWebSocketHttpMount PublicMount;
-		PublicMount.SetPathOnDisk(ServersDir / TEXT("Public"));
+		PublicMount.SetPathOnDisk(ServersDir / TEXT("www"));
 		PublicMount.SetWebPath(FString(TEXT("/")));
 		PublicMount.SetDefaultFile(FString(TEXT("player.html")));
 		MountsArr.Add(PublicMount);
diff --git a/Engine/Plugins/Media/PixelStreaming2/Source/PixelStreaming2Servers/Private/SignallingServer.cpp b/Engine/Plugins/Media/PixelStreaming2/Source/PixelStreaming2Servers/Private/SignallingServer.cpp
--- a/Engine/Plugins/Media/PixelStreaming2/Source/PixelStreaming2Servers/Private/SignallingServer.cpp
+++ b/Engine/Plugins/Media/PixelStreaming2/Source/PixelStreaming2Servers/Private/SignallingServer.cpp
@@ -158,7 +158,7 @@ namespace UE::PixelStreaming2Servers
 
 		// Add /Public
 		FWebSocketHttpMount PublicMount;
-		PublicMount.SetPathOnDisk(ServersDir / TEXT("Public"));
+		PublicMount.SetPathOnDisk(ServersDir / TEXT("www"));
 		PublicMount.SetWebPath(FString(TEXT("/")));
 		PublicMount.SetDefaultFile(FString(TEXT("player.html")));
 		MountsArr.Add(PublicMount);
