[PS2] Fix: Re-implement bitrate updating functionality

diff --git a/Engine/Plugins/Media/PixelStreaming2/Source/PixelStreaming2/Private/EpicRtcWebsocket.cpp b/Engine/Plugins/Media/PixelStreaming2/Source/PixelStreaming2/Private/EpicRtcWebsocket.cpp
--- a/Engine/Plugins/Media/PixelStreaming2/Source/PixelStreaming2/Private/EpicRtcWebsocket.cpp
+++ b/Engine/Plugins/Media/PixelStreaming2/Source/PixelStreaming2/Private/EpicRtcWebsocket.cpp
@@ -1,12 +1,14 @@
 // Copyright Epic Games, Inc. All Rights Reserved.
 
 #include "EpicRtcWebsocket.h"
-#include "PixelStreaming2PluginSettings.h"
+
+#include "IPixelStreaming2Module.h"
 #include "IWebSocket.h"
-#include "ToStringExtensions.h"
-#include "WebSocketsModule.h"
+#include "PixelStreaming2PluginSettings.h"
 #include "Serialization/JsonSerializer.h"
+#include "ToStringExtensions.h"
 #include "VideoUtils.h"
+#include "WebSocketsModule.h"
 
 namespace UE::PixelStreaming2
 {
@@ -85,17 +87,17 @@ namespace UE::PixelStreaming2
 		FString MessageString = FString{ (int32)Message._length, Message._ptr };
 
 		// Hijacking the offer message is a bit cheeky and should be removed once RTCP-7055 is closed.
-		TSharedPtr<FJsonObject> JsonObject = MakeShareable(new FJsonObject);
+		TSharedPtr<FJsonObject>	  JsonObject = MakeShareable(new FJsonObject);
 		TSharedRef<TJsonReader<>> JsonReader = TJsonReaderFactory<>::Create(MessageString);
 		if (FJsonSerializer::Deserialize(JsonReader, JsonObject) && JsonObject.IsValid())
 		{
 			FString MessageType;
 			JsonObject->TryGetStringField(TEXT("type"), MessageType);
 
-			if(MessageType == TEXT("offer"))
+			if (MessageType == TEXT("offer"))
 			{
-				EScalabilityMode ScalabilityMode = UE::PixelStreaming2::GetEnumFromCVar<EScalabilityMode>(UPixelStreaming2PluginSettings::CVarEncoderScalabilityMode);
-				FString ScalabilityModeString = UE::PixelStreaming2::GetCVarStringFromEnum<EScalabilityMode>(ScalabilityMode);
+				EScalabilityMode			 ScalabilityMode = UE::PixelStreaming2::GetEnumFromCVar<EScalabilityMode>(UPixelStreaming2PluginSettings::CVarEncoderScalabilityMode);
+				FString						 ScalabilityModeString = UE::PixelStreaming2::GetCVarStringFromEnum<EScalabilityMode>(ScalabilityMode);
 				TSharedRef<FJsonValueString> JsonValueObject = MakeShareable(new FJsonValueString(ScalabilityModeString));
 
 				JsonObject->SetField(TEXT("scalabilityMode"), JsonValueObject);
@@ -150,6 +152,34 @@ namespace UE::PixelStreaming2
 
 	void FEpicRtcWebsocket::OnMessage(const FString& Msg)
 	{
+		// Hijacking the answer message is a bit cheeky and should be removed once RTCP-7130 is closed.
+		TSharedPtr<FJsonObject>	  JsonObject = MakeShareable(new FJsonObject);
+		TSharedRef<TJsonReader<>> JsonReader = TJsonReaderFactory<>::Create(Msg);
+		if (FJsonSerializer::Deserialize(JsonReader, JsonObject) && JsonObject.IsValid())
+		{
+			FString MessageType;
+			JsonObject->TryGetStringField(TEXT("type"), MessageType);
+
+			if (MessageType == TEXT("answer"))
+			{
+				FString PlayerId;
+				if (JsonObject->TryGetStringField(TEXT("playerId"), PlayerId))
+				{
+					int		   MinBitrate;
+					int		   MaxBitrate;
+					const bool bGotMinBitrate = JsonObject->TryGetNumberField(TEXT("minBitrateBps"), MinBitrate);
+					const bool bGotMaxBitrate = JsonObject->TryGetNumberField(TEXT("maxBitrateBps"), MaxBitrate);
+
+					if (bGotMinBitrate && bGotMaxBitrate && MinBitrate > 0 && MaxBitrate > 0)
+					{
+						IPixelStreaming2Module::Get().ForEachStreamer([PlayerId, MinBitrate, MaxBitrate](TSharedPtr<IPixelStreaming2Streamer> Streamer) {
+							Streamer->PlayerRequestsBitrate(PlayerId, MinBitrate, MaxBitrate);
+						});
+					}
+				}
+			}
+		}
+
 		FUtf8String Message(Msg);
 		Observer->OnMessage(UE::PixelStreaming2::ToEpicRtcStringView(Message));
 	}
diff --git a/Engine/Plugins/Media/PixelStreaming2/Source/PixelStreaming2/Private/Streamer.cpp b/Engine/Plugins/Media/PixelStreaming2/Source/PixelStreaming2/Private/Streamer.cpp
--- a/Engine/Plugins/Media/PixelStreaming2/Source/PixelStreaming2/Private/Streamer.cpp
+++ b/Engine/Plugins/Media/PixelStreaming2/Source/PixelStreaming2/Private/Streamer.cpp
@@ -440,19 +440,34 @@ namespace UE::PixelStreaming2
 
 	void FStreamer::PlayerRequestsBitrate(FString PlayerId, int MinBitrate, int MaxBitrate)
 	{
-		// TODO (Migration): Reimplement once RTCP-6957 lands. This should only set the bitrate of the streams of this player (easier said than done)
-		UPixelStreaming2PluginSettings::CVarWebRTCMinBitrate->Set(MinBitrate, ECVF_SetByCode);
-		UPixelStreaming2PluginSettings::CVarWebRTCMaxBitrate->Set(MaxBitrate, ECVF_SetByCode);
+		UPixelStreaming2PluginSettings::CVarWebRTCMinBitrate.AsVariable()->Set(MinBitrate);
+		UPixelStreaming2PluginSettings::CVarWebRTCMaxBitrate.AsVariable()->Set(MaxBitrate);
 	}
 
 	void FStreamer::RefreshStreamBitrate()
 	{
 		Players->Apply([this](FString PlayerId, FPlayerContext& PlayerContext) {
-			// TODO (Migration) Reimplement once RTCP-6957 lands
-			// if (PlayerContext.PeerConnection)
-			// {
-			// PlayerContext.PeerConnection->RefreshStreamBitrate();
-			// }
+			if (!PlayerContext.ParticipantInterface)
+			{
+				return;
+			}
+
+			TRefCountPtr<EpicRtcConnectionInterface> ConnectionInterface = PlayerContext.ParticipantInterface->GetConnection();
+			if (!ConnectionInterface)
+			{
+				return;
+			}
+
+			EpicRtcBitrate Bitrates = {
+				._minBitrateBps = UPixelStreaming2PluginSettings::CVarWebRTCMinBitrate.GetValueOnAnyThread(),
+				._hasMinBitrateBps = true,
+				._maxBitrateBps = UPixelStreaming2PluginSettings::CVarWebRTCMaxBitrate.GetValueOnAnyThread(),
+				._hasMaxBitrateBps = true,
+				._startBitrateBps = UPixelStreaming2PluginSettings::CVarWebRTCStartBitrate.GetValueOnAnyThread(),
+				._hasStartBitrateBps = true
+			};
+
+			ConnectionInterface->SetConnectionRates(Bitrates);
 		});
 	}
 
@@ -1045,6 +1060,17 @@ namespace UE::PixelStreaming2
 
 		ParticipantConnection->AddDataSource(DataSource);
 
+		EpicRtcBitrate Bitrates = {
+			._minBitrateBps = UPixelStreaming2PluginSettings::CVarWebRTCMinBitrate.GetValueOnAnyThread(),
+			._hasMinBitrateBps = true,
+			._maxBitrateBps = UPixelStreaming2PluginSettings::CVarWebRTCMaxBitrate.GetValueOnAnyThread(),
+			._hasMaxBitrateBps = true,
+			._startBitrateBps = UPixelStreaming2PluginSettings::CVarWebRTCStartBitrate.GetValueOnAnyThread(),
+			._hasStartBitrateBps = true
+		};
+
+		ParticipantConnection->SetConnectionRates(Bitrates);
+
 		ParticipantConnection->StartNegotiation();
 	}
 
