[PS2] Fix: Peer component audio

diff --git a/Engine/Plugins/Media/PixelStreaming2/Source/PixelStreaming2/Private/Blueprints/PixelStreaming2Peer.cpp b/Engine/Plugins/Media/PixelStreaming2/Source/PixelStreaming2/Private/Blueprints/PixelStreaming2Peer.cpp
--- a/Engine/Plugins/Media/PixelStreaming2/Source/PixelStreaming2/Private/Blueprints/PixelStreaming2Peer.cpp
+++ b/Engine/Plugins/Media/PixelStreaming2/Source/PixelStreaming2/Private/Blueprints/PixelStreaming2Peer.cpp
@@ -217,7 +217,7 @@ void UPixelStreaming2Peer::OnEndGenerate()
 void UPixelStreaming2Peer::ConsumeRawPCM(const int16_t* AudioData, int InSampleRate, size_t NChannels, size_t NFrames)
 {
 	// Sound generator has not been initialized yet.
-	if (!SoundGenerator || GetAudioComponent() == nullptr)
+	if (SoundGenerator->GetSampleRate() == 0 || GetAudioComponent() == nullptr)
 	{
 		return;
 	}
@@ -365,7 +365,7 @@ void UPixelStreaming2Peer::OnAudioTrackFrame(EpicRtcAudioTrackInterface* AudioTr
 		return;
 	}
 
-	AudioSink->OnAudioData(Frame._data, Frame._format._sampleRate, Frame._format._numChannels, Frame._length);
+	AudioSink->OnAudioData(Frame._data, Frame._length, Frame._format._numChannels, Frame._format._sampleRate);
 }
 
 void UPixelStreaming2Peer::OnAudioTrackRemoved(EpicRtcAudioTrackInterface* AudioTrack)
