[PS2] Fix: Peer components crashing in -game due to EpicRtcConference not being initialized

diff --git a/Engine/Plugins/Media/PixelStreaming2/Source/PixelStreaming2/Private/PixelStreaming2Module.cpp b/Engine/Plugins/Media/PixelStreaming2/Source/PixelStreaming2/Private/PixelStreaming2Module.cpp
--- a/Engine/Plugins/Media/PixelStreaming2/Source/PixelStreaming2/Private/PixelStreaming2Module.cpp
+++ b/Engine/Plugins/Media/PixelStreaming2/Source/PixelStreaming2/Private/PixelStreaming2Module.cpp
@@ -137,7 +137,7 @@ namespace UE::PixelStreaming2
 		EpicRtcThread = MakeShared<FEpicRtcThread>();
 
 		// By calling InitDefaultStreamer post engine init we can use pixel streaming in standalone editor mode
-		FCoreDelegates::OnFEngineLoopInitComplete.AddLambda([this, RHIType]() {
+		FCoreDelegates::OnPostEngineInit.AddLambda([this, RHIType]() {
 			// Need to initialize after other modules have initialized such as NVCodec.
 			if (!InitializeEpicRtc())
 			{
