[PS2] Add: Initial stat support

diff --git a/Engine/Plugins/Media/PixelStreaming2/Source/PixelStreaming2/Internal/EpicRtcConferenceUtils.h b/Engine/Plugins/Media/PixelStreaming2/Source/PixelStreaming2/Internal/EpicRtcConferenceUtils.h
--- a/Engine/Plugins/Media/PixelStreaming2/Source/PixelStreaming2/Internal/EpicRtcConferenceUtils.h
+++ b/Engine/Plugins/Media/PixelStreaming2/Source/PixelStreaming2/Internal/EpicRtcConferenceUtils.h
@@ -3,6 +3,8 @@
 #pragma once
 
 #include "EpicRtcTickableTask.h"
+#include "IPixelStreaming2Stats.h"
+#include "PixelStreaming2PluginSettings.h"
 
 #include "epic_rtc/core/conference.h"
 
@@ -38,8 +40,10 @@ namespace UE::PixelStreaming2
 			if (EpicRtcConference)
 			{
 				MsSinceLastAudioTick += DeltaMs;
+				MsSinceLastStatsTick += DeltaMs;
 
 				// Tick conference normally. This handles things like data channel message
+				IPixelStreaming2Stats::Get().GraphValue(TEXT("ConferenceTickInterval"), DeltaMs, 1, 0.f, 1.f);
 				while (EpicRtcConference->NeedsTick())
 				{
 					EpicRtcConference->Tick();
@@ -48,9 +52,21 @@ namespace UE::PixelStreaming2
 				// Tick audio (every 10 ms). This enables the pulling of audio from the ADM
 				if (MsSinceLastAudioTick >= 10.f)
 				{
+					// Track the interval. This is helpful for seeing if we're going over the 10ms requirement
+					IPixelStreaming2Stats::Get().GraphValue(TEXT("AudioTickInterval"), MsSinceLastAudioTick, 1, 0.f, 25.f, 10.f);
 					EpicRtcConference->TickAudio();
 					MsSinceLastAudioTick = 0.f;
 				}
+
+				// Tick stats at the configured interval
+				float StatsInterval = UPixelStreaming2PluginSettings::CVarWebRTCStatsInterval.GetValueOnAnyThread() * 1000.f;
+				bool  bStatsEnabled = !UPixelStreaming2PluginSettings::CVarWebRTCDisableStats.GetValueOnAnyThread();
+				if (MsSinceLastStatsTick >= StatsInterval && StatsInterval > 0.f && bStatsEnabled)
+				{
+					IPixelStreaming2Stats::Get().GraphValue(TEXT("StatTickInterval"), MsSinceLastStatsTick, 1, 0.f, 25.f, 10.f);
+					EpicRtcConference->TickStats();
+					MsSinceLastStatsTick = 0.f;
+				}
 			}
 		}
 
@@ -65,6 +81,7 @@ namespace UE::PixelStreaming2
 		FString									 TaskName;
 
 		float MsSinceLastAudioTick = 0.f;
+		float MsSinceLastStatsTick = 0.f;
 	};
 
 } // namespace UE::PixelStreaming2
\ No newline at end of file
diff --git a/Engine/Plugins/Media/PixelStreaming2/Source/PixelStreaming2/Private/Blueprints/PixelStreaming2Peer.cpp b/Engine/Plugins/Media/PixelStreaming2/Source/PixelStreaming2/Private/Blueprints/PixelStreaming2Peer.cpp
--- a/Engine/Plugins/Media/PixelStreaming2/Source/PixelStreaming2/Private/Blueprints/PixelStreaming2Peer.cpp
+++ b/Engine/Plugins/Media/PixelStreaming2/Source/PixelStreaming2/Private/Blueprints/PixelStreaming2Peer.cpp
@@ -66,6 +66,8 @@ void UPixelStreaming2Peer::BeginPlay()
 	EpicRtcManager->VideoTrackObserverFactory = new UE::PixelStreaming2::FEpicRtcVideoTrackObserverFactory(EpicRtcManager);
 	EpicRtcManager->DataTrackObserverFactory = new UE::PixelStreaming2::FEpicRtcDataTrackObserverFactory(EpicRtcManager);
 
+	UE::PixelStreaming2::FPixelStreaming2Module::GetModule()->GetStatsCollector()->OnStatsReady.AddUObject(this, &UPixelStreaming2Peer::OnStatsReady);
+
 	Super::BeginPlay();
 }
 
@@ -193,6 +195,9 @@ bool UPixelStreaming2Peer::Subscribe(const FString& StreamerId)
 
 	EpicRtcManager->EpicRtcRoom->Join();
 
+	// Create a stats collector so we can receive stats from the subscribed streamer
+	StatsCollector = UE::PixelStreaming2::FRTCStatsCollector::Create(StreamerId);
+
 	return true;
 }
 
@@ -414,4 +419,20 @@ void UPixelStreaming2Peer::OnDataTrackState(EpicRtcDataTrackInterface* DataTrack
 
 void UPixelStreaming2Peer::OnDataTrackMessage(EpicRtcDataTrackInterface* DataTrack)
 {
+}
+
+void UPixelStreaming2Peer::OnStatsReady(const FString& PeerId, const EpicRtcConnectionStats& ConnectionStats)
+{
+	if (!StatsCollector)
+	{
+		return;
+	}
+
+	FString StreamId = *SubscribedStream;
+	if (PeerId != StreamId)
+	{
+		return;
+	}
+
+	StatsCollector->Process(ConnectionStats);
 }
\ No newline at end of file
diff --git a/Engine/Plugins/Media/PixelStreaming2/Source/PixelStreaming2/Private/Blueprints/PixelStreaming2Peer.h b/Engine/Plugins/Media/PixelStreaming2/Source/PixelStreaming2/Private/Blueprints/PixelStreaming2Peer.h
--- a/Engine/Plugins/Media/PixelStreaming2/Source/PixelStreaming2/Private/Blueprints/PixelStreaming2Peer.h
+++ b/Engine/Plugins/Media/PixelStreaming2/Source/PixelStreaming2/Private/Blueprints/PixelStreaming2Peer.h
@@ -2,12 +2,13 @@
 
 #pragma once
 
+#include "AudioSink.h"
 #include "Blueprints/PixelStreaming2MediaTexture.h"
 #include "Components/SynthComponent.h"
 #include "Containers/Utf8String.h"
-#include "IPixelStreaming2AudioConsumer.h"
-#include "AudioSink.h"
 #include "EpicRtcManager.h"
+#include "IPixelStreaming2AudioConsumer.h"
+#include "RTCStatsCollector.h"
 #include "VideoSink.h"
 
 #include "PixelStreaming2Peer.generated.h"
@@ -79,13 +80,15 @@ private:
 	FUtf8String		PlayerName;
 	static uint32_t PlayerId;
 
-	TSharedPtr<UE::PixelStreaming2::FAudioSink>							AudioSink;
+	TSharedPtr<UE::PixelStreaming2::FAudioSink>							  AudioSink;
 	TSharedPtr<UE::PixelStreaming2::FSoundGenerator, ESPMode::ThreadSafe> SoundGenerator;
 
 	TSharedPtr<UE::PixelStreaming2::FVideoSink> VideoSink;
 
 	TSharedPtr<UE::PixelStreaming2::FEpicRtcManager> EpicRtcManager;
 
+	TSharedPtr<UE::PixelStreaming2::FRTCStatsCollector> StatsCollector;
+
 public:
 	/* Begin IPixelStreaming2AudioConsumer Callbacks */
 	virtual void ConsumeRawPCM(const int16_t* AudioData, int InSampleRate, size_t NChannels, size_t NFrames) override;
@@ -123,6 +126,8 @@ public:
 	void OnDataTrackMessage(EpicRtcDataTrackInterface* DataTrack);
 	/* End FEpicRtcManager Callbacks */
 
+	void OnStatsReady(const FString& PeerId, const EpicRtcConnectionStats& ConnectionStats);
+
 private:
 	EpicRtcSessionState SessionState = EpicRtcSessionState::Disconnected;
 };
diff --git a/Engine/Plugins/Media/PixelStreaming2/Source/PixelStreaming2/Private/EpicRtcStatsCollector.cpp b/Engine/Plugins/Media/PixelStreaming2/Source/PixelStreaming2/Private/EpicRtcStatsCollector.cpp
new file mode 100644
--- /dev/null
+++ b/Engine/Plugins/Media/PixelStreaming2/Source/PixelStreaming2/Private/EpicRtcStatsCollector.cpp
@@ -0,0 +1,33 @@
+// Copyright Epic Games, Inc. All Rights Reserved.
+
+#include "EpicRtcStatsCollector.h"
+
+#include "ToStringExtensions.h"
+
+namespace UE::PixelStreaming2
+{
+	void FEpicRtcStatsCollector::OnStatsDelivered(const EpicRtcStatsReport& InReport)
+	{
+		/**
+		 * (Nazar.Rudenko): We care only for EpicRtcConnectionStats.
+		 * Every EpicRtcConnectionStats object corresponds to a Player/Streamer
+		 */
+		for (int s = 0; s < InReport._sessionStats._size; s++)
+		{
+			const EpicRtcSessionStats& Session = InReport._sessionStats._ptr[s];
+
+			for (int r = 0; r < Session._roomStats._size; r++)
+			{
+				const EpicRtcRoomStats& Room = Session._roomStats._ptr[r];
+
+				for (int c = 0; c < Room._connectionStats._size; c++)
+				{
+					const EpicRtcConnectionStats& Connection = Room._connectionStats._ptr[c];
+					FString						  CollectorId = ToString(Connection._connectionId);
+
+					OnStatsReady.Broadcast(CollectorId, Connection);
+				}
+			}
+		}
+	}
+} // namespace UE::PixelStreaming2
\ No newline at end of file
diff --git a/Engine/Plugins/Media/PixelStreaming2/Source/PixelStreaming2/Private/EpicRtcStatsCollector.h b/Engine/Plugins/Media/PixelStreaming2/Source/PixelStreaming2/Private/EpicRtcStatsCollector.h
new file mode 100644
--- /dev/null
+++ b/Engine/Plugins/Media/PixelStreaming2/Source/PixelStreaming2/Private/EpicRtcStatsCollector.h
@@ -0,0 +1,29 @@
+// Copyright Epic Games, Inc. All Rights Reserved.
+
+#pragma once
+
+#include "RTCStatsCollector.h"
+#include "Templates/RefCounting.h"
+#include "Templates/SharedPointer.h"
+
+#include "epic_rtc/core/stats.h"
+
+namespace UE::PixelStreaming2
+{
+	class FEpicRtcStatsCollector : public EpicRtcStatsCollectorCallbackInterface, public TRefCountingMixin<FEpicRtcStatsCollector>
+	{
+	public:
+		FEpicRtcStatsCollector() = default;
+		~FEpicRtcStatsCollector() = default;
+
+		// Begin EpicRtcStatsCollectorCallbackInterface interface
+		void			 OnStatsDelivered(const EpicRtcStatsReport& InReport) override;
+		virtual uint32_t AddRef() override final { return TRefCountingMixin<FEpicRtcStatsCollector>::AddRef(); }
+		virtual uint32_t Release() override final { return TRefCountingMixin<FEpicRtcStatsCollector>::Release(); }
+		virtual uint32_t Count() const override final { return TRefCountingMixin<FEpicRtcStatsCollector>::GetRefCount(); }
+		// End EpicRtcStatsCollectorCallbackInterface interface
+
+		DECLARE_EVENT_TwoParams(FEpicRtcStatsCollector, FOnStatsReady, const FString&, const EpicRtcConnectionStats&);
+		FOnStatsReady OnStatsReady;
+	};
+} // namespace UE::PixelStreaming2
\ No newline at end of file
diff --git a/Engine/Plugins/Media/PixelStreaming2/Source/PixelStreaming2/Private/PixelStreaming2Module.cpp b/Engine/Plugins/Media/PixelStreaming2/Source/PixelStreaming2/Private/PixelStreaming2Module.cpp
--- a/Engine/Plugins/Media/PixelStreaming2/Source/PixelStreaming2/Private/PixelStreaming2Module.cpp
+++ b/Engine/Plugins/Media/PixelStreaming2/Source/PixelStreaming2/Private/PixelStreaming2Module.cpp
@@ -663,6 +663,8 @@ namespace UE::PixelStreaming2
 
 		WebsocketFactory = MakeRefCount<FEpicRtcWebsocketFactory>();
 
+		StatsCollector = MakeRefCount<FEpicRtcStatsCollector>();
+
 		// clang-format off
 		EpicRtcConfig ConferenceConfig = {
 			._websocketFactory = WebsocketFactory.GetReference(),
@@ -697,6 +699,11 @@ namespace UE::PixelStreaming2
 				._level = UnrealLogToEpicRtcCategoryMap[LogPixelStreaming2EpicRtc.GetVerbosity()],
 				._levelWebRtc = UnrealLogToEpicRtcCategoryMap[LogPixelStreaming2EpicRtc.GetVerbosity()]
 #endif
+			},
+			._stats = {
+				._statsCollectorCallback = StatsCollector.GetReference(),
+				._statsCollectorInterval = 1000,
+				._jsonFormatOnly = false
 			}
 		};
 		// clang-format on
diff --git a/Engine/Plugins/Media/PixelStreaming2/Source/PixelStreaming2/Private/PixelStreaming2Module.h b/Engine/Plugins/Media/PixelStreaming2/Source/PixelStreaming2/Private/PixelStreaming2Module.h
--- a/Engine/Plugins/Media/PixelStreaming2/Source/PixelStreaming2/Private/PixelStreaming2Module.h
+++ b/Engine/Plugins/Media/PixelStreaming2/Source/PixelStreaming2/Private/PixelStreaming2Module.h
@@ -4,6 +4,7 @@
 
 #include "IPixelStreaming2Module.h"
 #include "EpicRtcConferenceUtils.h"
+#include "EpicRtcStatsCollector.h"
 
 #include "epic_rtc/core/platform.h"
 #include "epic_rtc/plugins/signalling/signalling_type.h"
@@ -44,14 +45,15 @@ namespace UE::PixelStreaming2
 		virtual FString									  GetDefaultSignallingURL() override;
 
 		// These are staying on the module at the moment as theres no way of the BPs knowing which streamer they are relevant to
-		virtual void										  AddInputComponent(UPixelStreaming2Input* InInputComponent);
-		virtual void										  RemoveInputComponent(UPixelStreaming2Input* InInputComponent);
+		virtual void								 AddInputComponent(UPixelStreaming2Input* InInputComponent);
+		virtual void								 RemoveInputComponent(UPixelStreaming2Input* InInputComponent);
 		virtual const TArray<UPixelStreaming2Input*> GetInputComponents();
-		virtual void										  ForEachStreamer(const TFunction<void(TSharedPtr<IPixelStreaming2Streamer>)>& Func) override;
+		virtual void								 ForEachStreamer(const TFunction<void(TSharedPtr<IPixelStreaming2Streamer>)>& Func) override;
 		/** End IPixelStreaming2Module implementation */
 
 		TSharedPtr<class FEpicRtcAudioMixingCapturer> GetAudioCapturer();
 		TRefCountPtr<EpicRtcConferenceInterface>	  GetEpicRtcConference() { return EpicRtcConference; }
+		TRefCountPtr<FEpicRtcStatsCollector>		  GetStatsCollector() { return StatsCollector; }
 
 	private:
 		/** IModuleInterface implementation */
@@ -76,7 +78,7 @@ namespace UE::PixelStreaming2
 		FReadyEvent		ReadyEvent;
 		FDelegateHandle LogStatsHandle;
 
-		TArray<UPixelStreaming2Input*>			  InputComponents;
+		TArray<UPixelStreaming2Input*>					  InputComponents;
 		mutable FCriticalSection						  StreamersCS;
 		TMap<FString, TWeakPtr<IPixelStreaming2Streamer>> Streamers;
 		TSharedPtr<IPixelStreaming2Streamer>			  DefaultStreamer;
@@ -87,6 +89,7 @@ namespace UE::PixelStreaming2
 		TSharedPtr<class FEpicRtcAudioMixingCapturer> AudioMixingCapturer;
 		TRefCountPtr<EpicRtcPlatformInterface>		  EpicRtcPlatform;
 		TRefCountPtr<EpicRtcConferenceInterface>	  EpicRtcConference;
+		TRefCountPtr<FEpicRtcStatsCollector>		  StatsCollector;
 
 		TRefCountPtr<FEpicRtcWebsocketFactory>	   WebsocketFactory;
 		TUniqueTaskPtr<FEpicRtcTickConferenceTask> TickConferenceTask;
diff --git a/Engine/Plugins/Media/PixelStreaming2/Source/PixelStreaming2/Private/PlayerContext.h b/Engine/Plugins/Media/PixelStreaming2/Source/PixelStreaming2/Private/PlayerContext.h
--- a/Engine/Plugins/Media/PixelStreaming2/Source/PixelStreaming2/Private/PlayerContext.h
+++ b/Engine/Plugins/Media/PixelStreaming2/Source/PixelStreaming2/Private/PlayerContext.h
@@ -7,6 +7,7 @@
 #include "EpicRtcDataTrack.h"
 #include "EpicRtcVideoSink.h"
 #include "EpicRtcVideoSource.h"
+#include "RTCStatsCollector.h"
 #include "Templates/RefCounting.h"
 #include "Utils.h"
 
@@ -28,6 +29,8 @@ namespace UE::PixelStreaming2
 		TSharedPtr<FEpicRtcVideoSink>	VideoSink;
 
 		TSharedPtr<FEpicRtcDataTrack> DataTrack;
+
+		TSharedPtr<FRTCStatsCollector> StatsCollector;
 	};
 
 } // namespace UE::PixelStreaming2
\ No newline at end of file
diff --git a/Engine/Plugins/Media/PixelStreaming2/Source/PixelStreaming2/Private/RTCStatsCollector.cpp b/Engine/Plugins/Media/PixelStreaming2/Source/PixelStreaming2/Private/RTCStatsCollector.cpp
--- a/Engine/Plugins/Media/PixelStreaming2/Source/PixelStreaming2/Private/RTCStatsCollector.cpp
+++ b/Engine/Plugins/Media/PixelStreaming2/Source/PixelStreaming2/Private/RTCStatsCollector.cpp
@@ -1,187 +1,688 @@
 // Copyright Epic Games, Inc. All Rights Reserved.
 
 #include "RTCStatsCollector.h"
+
 #include "Logging.h"
+#include "PixelStreaming2PluginSettings.h"
 #include "PixelStreaming2StatNames.h"
+#include "Streamer.h"
 #include "ToStringExtensions.h"
 
-// TODO (Migration) RTCP-6453 Retrieve stats from EpicRtc when EpicRtc supports it
-#if 0
-
 namespace UE::PixelStreaming2
 {
+	TSharedPtr<FRTCStatsCollector> FRTCStatsCollector::Create(const FString& PlayerId)
+	{
+		TSharedPtr<FRTCStatsCollector> StatsCollector = TSharedPtr<FRTCStatsCollector>(new FRTCStatsCollector(PlayerId));
+
+		if (UPixelStreaming2PluginSettings::FDelegates* Delegates = UPixelStreaming2PluginSettings::Delegates())
+		{
+			Delegates->OnWebRTCDisableStatsChanged.AddSP(StatsCollector.ToSharedRef(), &FRTCStatsCollector::OnWebRTCDisableStatsChanged);
+		}
 
-	// ------------- FRTCStatsCollector-------------------
+		return StatsCollector;
+	}
 
 	FRTCStatsCollector::FRTCStatsCollector()
 		: FRTCStatsCollector(INVALID_PLAYER_ID)
 	{
 	}
 
-	FRTCStatsCollector::FRTCStatsCollector(FPixelStreaming2PlayerId PlayerId)
+	FRTCStatsCollector::FRTCStatsCollector(const FString& PlayerId)
 		: AssociatedPlayerId(PlayerId)
 		, LastCalculationCycles(FPlatformTime::Cycles64())
-		, bIsEnabled(!Settings::CVarPixelStreaming2WebRTCDisableStats.GetValueOnAnyThread())
+		, bIsEnabled(!UPixelStreaming2PluginSettings::CVarWebRTCDisableStats.GetValueOnAnyThread())
+		, LocalVideoTrackStatsSink(MakeUnique<FRTPLocalVideoTrackStatsSink>())
+		, LocalAudioTrackStatsSink(MakeUnique<FRTPLocalAudioTrackStatsSink>())
+		, VideoSourceStatsSink(MakeUnique<FVideoSourceStatsSink>())
+		, AudioSourceStatsSink(MakeUnique<FAudioSourceStatsSink>())
+		, RemoteVideoTrackStatsSink(MakeUnique<FRTPRemoteTrackStatsSink>(RTCStatCategories::RemoteVideoTrack))
+		, RemoteAudioTrackStatsSink(MakeUnique<FRTPRemoteTrackStatsSink>(RTCStatCategories::RemoteAudioTrack))
+		, DataTrackStatsSink(MakeUnique<FDataTrackStatsSink>())
+		, CandidatePairStatsSink(MakeUnique<FCandidatePairStatsSink>())
 	{
-
-		// Add each of the sinks for each of the stat types we are interested in tracking
-		RTCStatSinks.Add(MakeShared<FTrackStatsSink>());
-		RTCStatSinks.Add(MakeShared<FRTPMediaStatsSink>(RTCStatTypes::InboundRTP, "video"));
-		RTCStatSinks.Add(MakeShared<FRTPMediaStatsSink>(RTCStatTypes::OutboundRTP, "video"));
-		RTCStatSinks.Add(MakeShared<FRTPMediaStatsSink>(RTCStatTypes::InboundRTP, "audio"));
-		RTCStatSinks.Add(MakeShared<FRTPMediaStatsSink>(RTCStatTypes::OutboundRTP, "audio"));
-		RTCStatSinks.Add(MakeShared<FVideoSourceStatsSink>());
-		RTCStatSinks.Add(MakeShared<FDataChannelStatsSink>());
 	}
 
-	void FRTCStatsCollector::AddRef() const
+	void FRTCStatsCollector::OnWebRTCDisableStatsChanged(IConsoleVariable* Var)
 	{
-		FPlatformAtomics::InterlockedIncrement(&RefCount);
+		bIsEnabled = !Var->GetBool();
 	}
 
-	rtc::RefCountReleaseStatus FRTCStatsCollector::Release() const
+	void FRTCStatsCollector::Process(const EpicRtcConnectionStats& InStats)
 	{
-		if (FPlatformAtomics::InterlockedDecrement(&RefCount) == 0)
+		FStats* PSStats = FStats::Get();
+		if (!bIsEnabled || !PSStats)
 		{
-			return rtc::RefCountReleaseStatus::kDroppedLastRef;
+			return;
+		}
+
+		uint64 CyclesNow = FPlatformTime::Cycles64();
+		double SecondsDelta = FGenericPlatformTime::ToSeconds64(CyclesNow - LastCalculationCycles);
+
+		if (InStats._localVideoTracks._size > 0)
+		{
+			if (InStats._localVideoTracks._size > 1)
+			{
+				UE_LOGFMT(LogPixelStreaming2, Warning, "Multiple local video tracks is unsupported. Expected: 1, Actual: {0}", InStats._localVideoTracks._size);
+			}
+			LocalVideoTrackStatsSink->Process(InStats._localVideoTracks._ptr[0], AssociatedPlayerId, SecondsDelta);
+			VideoSourceStatsSink->Process(InStats._localVideoTracks._ptr[0]._source, AssociatedPlayerId, SecondsDelta);
+		}
+
+		if (InStats._localAudioTracks._size > 0)
+		{
+			if (InStats._localAudioTracks._size > 1)
+			{
+				UE_LOGFMT(LogPixelStreaming2, Warning, "Multiple local audio tracks is unsupported. Expected: 1, Actual: {0}", InStats._localAudioTracks._size);
+			}
+			LocalAudioTrackStatsSink->Process(InStats._localAudioTracks._ptr[0]._rtp, AssociatedPlayerId, SecondsDelta);
+			AudioSourceStatsSink->Process(InStats._localAudioTracks._ptr[0]._source, AssociatedPlayerId, SecondsDelta);
+		}
+
+		if (InStats._remoteVideoTracks._size > 0)
+		{
+			if (InStats._remoteVideoTracks._size > 1)
+			{
+				UE_LOGFMT(LogPixelStreaming2, Warning, "Multiple remote video tracks is unsupported. Expected: 1, Actual: {0}", InStats._remoteVideoTracks._size);
+			}
+			RemoteVideoTrackStatsSink->Process(InStats._remoteVideoTracks._ptr[0]._rtp, AssociatedPlayerId, SecondsDelta);
+		}
+
+		if (InStats._remoteAudioTracks._size > 0)
+		{
+			if (InStats._remoteAudioTracks._size > 1)
+			{
+				UE_LOGFMT(LogPixelStreaming2, Warning, "Multiple remote audio tracks is unsupported. Expected: 1, Actual: {0}", InStats._remoteAudioTracks._size);
+			}
+			RemoteAudioTrackStatsSink->Process(InStats._remoteAudioTracks._ptr[0]._rtp, AssociatedPlayerId, SecondsDelta);
+		}
+
+		if (InStats._dataTracks._size > 0)
+		{
+			if (InStats._dataTracks._size > 1)
+			{
+				UE_LOGFMT(LogPixelStreaming2, Warning, "Multiple data tracks is unsupported. Expected: 1, Actual: {0}", InStats._dataTracks._size);
+			}
+			DataTrackStatsSink->Process(InStats._dataTracks._ptr[0], AssociatedPlayerId, SecondsDelta);
+		}
+
+		if (InStats._transports._size > 0)
+		{
+			//(Nazar.Rudenko): More than one transport is possible only if we are not using bundle which we do
+			const EpicRtcTransportStats& Transport = InStats._transports._ptr[0];
+			FString						 SelectedPairId = ToString(Transport._selectedCandidatePairId);
+			for (int i = 0; i < Transport._candidatePairs._size; i++)
+			{
+				FString PairId = ToString(Transport._candidatePairs._ptr[i]._id);
+				if (SelectedPairId == PairId)
+				{
+					CandidatePairStatsSink->Process(Transport._candidatePairs._ptr[i], AssociatedPlayerId, SecondsDelta);
+				}
+			}
 		}
 
-		return rtc::RefCountReleaseStatus::kOtherRefsRemained;
+		LastCalculationCycles = FPlatformTime::Cycles64();
 	}
+	/*
+	 * ---------- FRTPLocalVideoTrackSink ----------
+	 */
+	FRTCStatsCollector::FRTPLocalVideoTrackStatsSink::FRTPLocalVideoTrackStatsSink()
+		: FStatsSink(RTCStatCategories::LocalVideoTrack)
+	{
+		// These stats will be extracted from the stat reports and emitted straight to screen
+		Add(PixelStreaming2StatNames::FirCount, 0);
+		Add(PixelStreaming2StatNames::PliCount, 0);
+		Add(PixelStreaming2StatNames::NackCount, 0);
+		Add(PixelStreaming2StatNames::RetransmittedBytesSent, 0);
+		Add(PixelStreaming2StatNames::TotalEncodeBytesTarget, 0);
+		Add(PixelStreaming2StatNames::KeyFramesEncoded, 0);
+		Add(PixelStreaming2StatNames::FrameWidth, 0);
+		Add(PixelStreaming2StatNames::FrameHeight, 0);
+		Add(PixelStreaming2StatNames::HugeFramesSent, 0);
+		Add(PixelStreaming2StatNames::PacketsLost, 0);
+		Add(PixelStreaming2StatNames::Jitter, 0);
+		Add(PixelStreaming2StatNames::RoundTripTime, 0);
+
+		// These are values used to calculate extra values (stores time deltas etc)
+		AddNonRendered(PixelStreaming2StatNames::TargetBitrate);
+		AddNonRendered(PixelStreaming2StatNames::FramesSent);
+		AddNonRendered(PixelStreaming2StatNames::FramesReceived);
+		AddNonRendered(PixelStreaming2StatNames::BytesSent);
+		AddNonRendered(PixelStreaming2StatNames::BytesReceived);
+		AddNonRendered(PixelStreaming2StatNames::QPSum);
+		AddNonRendered(PixelStreaming2StatNames::TotalEncodeTime);
+		AddNonRendered(PixelStreaming2StatNames::FramesEncoded);
+		AddNonRendered(PixelStreaming2StatNames::FramesDecoded);
+		AddNonRendered(PixelStreaming2StatNames::TotalPacketSendDelay);
+
+		// Calculated stats below:
+		// FrameSent Per Second
+		AddStatCalculator([](FStatsSink& StatSource, double Period) -> TOptional<FStatData> {
+			FRTCTrackedStat* FramesSentStat = StatSource.Get(PixelStreaming2StatNames::FramesSent);
+			if (FramesSentStat && FramesSentStat->GetLatestStat().StatValue > 0)
+			{
+				const double FramesSentPerSecond = FramesSentStat->CalculateDelta(Period);
+				FStatData	 FpsStat = FStatData(PixelStreaming2StatNames::FramesSentPerSecond, FramesSentPerSecond, 0);
+				FpsStat.DisplayFlags = FStatData::EDisplayFlags::TEXT | FStatData::EDisplayFlags::GRAPH;
+				return FpsStat;
+			}
+			return {};
+		});
 
-	void FRTCStatsCollector::OnStatsDelivered(const rtc::scoped_refptr<const webrtc::RTCStatsReport>& Report)
+		// FramesReceived Per Second
+		AddStatCalculator([](FStatsSink& StatSource, double Period) -> TOptional<FStatData> {
+			FRTCTrackedStat* FramesReceivedStat = StatSource.Get(PixelStreaming2StatNames::FramesReceived);
+			if (FramesReceivedStat && FramesReceivedStat->GetLatestStat().StatValue > 0)
+			{
+				const double FramesReceivedPerSecond = FramesReceivedStat->CalculateDelta(Period);
+				return FStatData(PixelStreaming2StatNames::FramesReceivedPerSecond, FramesReceivedPerSecond, 0);
+			}
+			return {};
+		});
+
+		// Megabits sent Per Second
+		AddStatCalculator([](FStatsSink& StatSource, double Period) -> TOptional<FStatData> {
+			FRTCTrackedStat* BytesSentStat = StatSource.Get(PixelStreaming2StatNames::BytesSent);
+			if (BytesSentStat && BytesSentStat->GetLatestStat().StatValue > 0)
+			{
+				const double BytesSentPerSecond = BytesSentStat->CalculateDelta(Period);
+				const double MegabitsPerSecond = BytesSentPerSecond / 1'000'000.0 * 8.0;
+				return FStatData(PixelStreaming2StatNames::BitrateMegabits, MegabitsPerSecond, 2);
+			}
+			return {};
+		});
+
+		// Bits sent Per Second
+		AddStatCalculator([](FStatsSink& StatSource, double Period) -> TOptional<FStatData> {
+			FRTCTrackedStat* BytesSentStat = StatSource.Get(PixelStreaming2StatNames::BytesSent);
+			if (BytesSentStat && BytesSentStat->GetLatestStat().StatValue > 0)
+			{
+				const double BytesSentPerSecond = BytesSentStat->CalculateDelta(Period);
+				const double BitsPerSecond = BytesSentPerSecond * 8.0;
+				FStatData	 Stat = FStatData(PixelStreaming2StatNames::Bitrate, BitsPerSecond, 0);
+				Stat.DisplayFlags = FStatData::EDisplayFlags::HIDDEN; // We don't want to display bits per second (too many digits)
+				return Stat;
+			}
+			return {};
+		});
+
+		// Target megabits sent Per Second
+		AddStatCalculator([](FStatsSink& StatSource, double Period) -> TOptional<FStatData> {
+			FRTCTrackedStat* TargetBpsStats = StatSource.Get(PixelStreaming2StatNames::TargetBitrate);
+			if (TargetBpsStats && TargetBpsStats->GetLatestStat().StatValue > 0)
+			{
+				const double TargetBps = TargetBpsStats->Average();
+				const double MegabitsPerSecond = TargetBps / 1'000'000.0;
+				return FStatData(PixelStreaming2StatNames::TargetBitrateMegabits, MegabitsPerSecond, 2);
+			}
+			return {};
+		});
+
+		// Megabits received Per Second
+		AddStatCalculator([](FStatsSink& StatSource, double Period) -> TOptional<FStatData> {
+			FRTCTrackedStat* BytesReceivedStat = StatSource.Get(PixelStreaming2StatNames::BytesReceived);
+			if (BytesReceivedStat && BytesReceivedStat->GetLatestStat().StatValue > 0)
+			{
+				const double BytesReceivedPerSecond = BytesReceivedStat->CalculateDelta(Period);
+				const double MegabitsPerSecond = BytesReceivedPerSecond / 1000.0 * 8.0;
+				return FStatData(PixelStreaming2StatNames::Bitrate, MegabitsPerSecond, 2);
+			}
+			return {};
+		});
+
+		// Encoded fps
+		AddStatCalculator([](FStatsSink& StatSource, double Period) -> TOptional<FStatData> {
+			FRTCTrackedStat* EncodedFramesStat = StatSource.Get(PixelStreaming2StatNames::FramesEncoded);
+			if (EncodedFramesStat && EncodedFramesStat->GetLatestStat().StatValue > 0)
+			{
+				const double EncodedFramesPerSecond = EncodedFramesStat->CalculateDelta(Period);
+				return FStatData(PixelStreaming2StatNames::EncodedFramesPerSecond, EncodedFramesPerSecond, 0);
+			}
+			return {};
+		});
+
+		// Decoded fps
+		AddStatCalculator([](FStatsSink& StatSource, double Period) -> TOptional<FStatData> {
+			FRTCTrackedStat* DecodedFramesStat = StatSource.Get(PixelStreaming2StatNames::FramesDecoded);
+			if (DecodedFramesStat && DecodedFramesStat->GetLatestStat().StatValue > 0)
+			{
+				const double DecodedFramesPerSecond = DecodedFramesStat->CalculateDelta(Period);
+				return FStatData(PixelStreaming2StatNames::DecodedFramesPerSecond, DecodedFramesPerSecond, 0);
+			}
+			return {};
+		});
+
+		// Avg QP Per Second
+		AddStatCalculator([](FStatsSink& StatSource, double Period) -> TOptional<FStatData> {
+			FRTCTrackedStat* QPSumStat = StatSource.Get(PixelStreaming2StatNames::QPSum);
+			FStatData*		 EncodedFramesPerSecond = StatSource.GetCalculatedStat(PixelStreaming2StatNames::EncodedFramesPerSecond);
+			if (QPSumStat && QPSumStat->GetLatestStat().StatValue > 0
+				&& EncodedFramesPerSecond && EncodedFramesPerSecond->StatValue > 0.0)
+			{
+				const double QPSumDeltaPerSecond = QPSumStat->CalculateDelta(Period);
+				const double MeanQPPerFrame = QPSumDeltaPerSecond / EncodedFramesPerSecond->StatValue;
+				FName		 StatName = PixelStreaming2StatNames::MeanQPPerSecond;
+				return FStatData(StatName, MeanQPPerFrame, 0);
+			}
+			return {};
+		});
+
+		// Mean EncodeTime (ms) Per Frame
+		AddStatCalculator([](FStatsSink& StatSource, double Period) -> TOptional<FStatData> {
+			FRTCTrackedStat* TotalEncodeTimeStat = StatSource.Get(PixelStreaming2StatNames::TotalEncodeTime);
+			FStatData*		 EncodedFramesPerSecond = StatSource.GetCalculatedStat(PixelStreaming2StatNames::EncodedFramesPerSecond);
+			if (TotalEncodeTimeStat && TotalEncodeTimeStat->GetLatestStat().StatValue > 0
+				&& EncodedFramesPerSecond && EncodedFramesPerSecond->StatValue > 0.0)
+			{
+				const double TotalEncodeTimePerSecond = TotalEncodeTimeStat->CalculateDelta(Period);
+				const double MeanEncodeTimePerFrameMs = TotalEncodeTimePerSecond / EncodedFramesPerSecond->StatValue * 1000.0;
+				return FStatData(PixelStreaming2StatNames::MeanEncodeTime, MeanEncodeTimePerFrameMs, 2);
+			}
+			return {};
+		});
+
+		// Mean SendDelay (ms) Per Frame
+		AddStatCalculator([](FStatsSink& StatSource, double Period) -> TOptional<FStatData> {
+			FRTCTrackedStat* TotalSendDelayStat = StatSource.Get(PixelStreaming2StatNames::TotalPacketSendDelay);
+			FStatData*		 FramesSentPerSecond = StatSource.GetCalculatedStat(PixelStreaming2StatNames::FramesSentPerSecond);
+			if (TotalSendDelayStat && TotalSendDelayStat->GetLatestStat().StatValue > 0
+				&& FramesSentPerSecond && FramesSentPerSecond->StatValue > 0.0)
+			{
+				const double TotalSendDelayPerSecond = TotalSendDelayStat->CalculateDelta(Period);
+				const double MeanSendDelayPerFrameMs = TotalSendDelayPerSecond / FramesSentPerSecond->StatValue * 1000.0;
+				return FStatData(PixelStreaming2StatNames::MeanSendDelay, MeanSendDelayPerFrameMs, 2);
+			}
+			return {};
+		});
+
+		// JitterBufferDelay (ms)
+		AddStatCalculator([](FStatsSink& StatSource, double Period) -> TOptional<FStatData> {
+			FRTCTrackedStat* JitterBufferDelayStat = StatSource.Get(PixelStreaming2StatNames::JitterBufferDelay);
+			FStatData*		 FramesReceivedPerSecond = StatSource.GetCalculatedStat(PixelStreaming2StatNames::FramesReceivedPerSecond);
+			if (JitterBufferDelayStat && JitterBufferDelayStat->GetLatestStat().StatValue > 0
+				&& FramesReceivedPerSecond && FramesReceivedPerSecond->StatValue > 0.0)
+			{
+				const double TotalJitterBufferDelayPerSecond = JitterBufferDelayStat->CalculateDelta(Period);
+				const double MeanJitterBufferDelayMs = TotalJitterBufferDelayPerSecond / FramesReceivedPerSecond->StatValue * 1000.0;
+				return FStatData(PixelStreaming2StatNames::JitterBufferDelay, MeanJitterBufferDelayMs, 2);
+			}
+			return {};
+		});
+	}
+
+	void FRTCStatsCollector::FRTPLocalVideoTrackStatsSink::Process(const EpicRtcLocalVideoTrackStats& InStats, const FString& PeerId, double SecondsDelta)
 	{
 		FStats* PSStats = FStats::Get();
+		if (!PSStats)
+		{
+			return;
+		}
 
-		if (!bIsEnabled || !PSStats || !Report)
+		if (InStats._rtp._size == 0)
 		{
 			return;
 		}
 
-		uint64 CyclesNow = FPlatformTime::Cycles64();
-		double SecondsDelta = FGenericPlatformTime::ToSeconds64(CyclesNow - LastCalculationCycles);
+		// Huge assumption, but assume that normal peers only have 1 video rtp stats and simulcast is anything > 1
+		bool bIsSFU = InStats._rtp._size > 1;
+
+		for (int i = 0; i < InStats._rtp._size; i++)
+		{
+			const EpicRtcLocalTrackRtpStats& RtpStats = InStats._rtp._ptr[i];
+			const FString					 StatsId = bIsSFU ? FString::FromInt(RtpStats._local._ssrc) : PeerId;
+
+			for (TPair<FName, FRTCTrackedStat>& Tuple : Stats)
+			{
+				double NewValue = 0;
+				if (Tuple.Key == PixelStreaming2StatNames::FirCount)
+				{
+					NewValue = RtpStats._local._firCount;
+				}
+				else if (Tuple.Key == PixelStreaming2StatNames::PliCount)
+				{
+					NewValue = RtpStats._local._pliCount;
+				}
+				else if (Tuple.Key == PixelStreaming2StatNames::NackCount)
+				{
+					NewValue = RtpStats._local._nackCount;
+				}
+				else if (Tuple.Key == PixelStreaming2StatNames::RetransmittedBytesSent)
+				{
+					NewValue = RtpStats._local._retransmittedBytesSent;
+				}
+				else if (Tuple.Key == PixelStreaming2StatNames::TotalEncodeBytesTarget)
+				{
+					NewValue = RtpStats._local._totalEncodedBytesTarget;
+				}
+				else if (Tuple.Key == PixelStreaming2StatNames::KeyFramesEncoded)
+				{
+					NewValue = RtpStats._local._keyFramesEncoded;
+				}
+				else if (Tuple.Key == PixelStreaming2StatNames::FrameWidth)
+				{
+					NewValue = RtpStats._local._frameWidth;
+				}
+				else if (Tuple.Key == PixelStreaming2StatNames::FrameHeight)
+				{
+					NewValue = RtpStats._local._frameHeight;
+				}
+				else if (Tuple.Key == PixelStreaming2StatNames::HugeFramesSent)
+				{
+					NewValue = RtpStats._local._hugeFramesSent;
+				}
+				else if (Tuple.Key == PixelStreaming2StatNames::TotalPacketSendDelay)
+				{
+					NewValue = RtpStats._local._totalPacketSendDelay;
+				}
+				else if (Tuple.Key == PixelStreaming2StatNames::TargetBitrate)
+				{
+					NewValue = RtpStats._local._targetBitrate;
+				}
+				else if (Tuple.Key == PixelStreaming2StatNames::FramesSent)
+				{
+					NewValue = RtpStats._local._framesSent;
+				}
+				else if (Tuple.Key == PixelStreaming2StatNames::FramesReceived)
+				{
+					// TODO(Nazar.Rudenko): Available for inbound tracks only
+				}
+				else if (Tuple.Key == PixelStreaming2StatNames::BytesSent)
+				{
+					NewValue = RtpStats._local._bytesSent;
+				}
+				else if (Tuple.Key == PixelStreaming2StatNames::BytesReceived)
+				{
+					// TODO(Nazar.Rudenko): Available for inbound tracks only
+				}
+				else if (Tuple.Key == PixelStreaming2StatNames::QPSum)
+				{
+					NewValue = RtpStats._local._qpSum;
+				}
+				else if (Tuple.Key == PixelStreaming2StatNames::TotalEncodeTime)
+				{
+					NewValue = RtpStats._local._totalEncodeTime;
+				}
+				else if (Tuple.Key == PixelStreaming2StatNames::FramesEncoded)
+				{
+					NewValue = RtpStats._local._framesEncoded;
+				}
+				else if (Tuple.Key == PixelStreaming2StatNames::FramesDecoded)
+				{
+					// TODO(Nazar.Rudenko): Available for inbound tracks only
+				}
+				else if (Tuple.Key == PixelStreaming2StatNames::PacketsLost)
+				{
+					NewValue = RtpStats._remote._packetsLost;
+				}
+				else if (Tuple.Key == PixelStreaming2StatNames::Jitter)
+				{
+					NewValue = RtpStats._remote._jitter;
+				}
+				else if (Tuple.Key == PixelStreaming2StatNames::RoundTripTime)
+				{
+					NewValue = RtpStats._remote._roundTripTime;
+				}
+
+				if (UpdateValue(NewValue, &Tuple.Value))
+				{
+					PSStats->StorePeerStat(StatsId, RTCStatCategories::LocalVideoTrack, Tuple.Value.GetLatestStat());
+				}
+			}
+
+			PostProcess(PSStats, StatsId, SecondsDelta);
+		}
+	}
+	/*
+	 * ---------- FRTPLocalAudioTrackStatsSink ----------
+	 */
+	FRTCStatsCollector::FRTPLocalAudioTrackStatsSink::FRTPLocalAudioTrackStatsSink()
+		: FStatsSink(RTCStatCategories::LocalAudioTrack)
+	{
+		// These stats will be extracted from the stat reports and emitted straight to screen
+		Add(PixelStreaming2StatNames::FirCount, 0);
+		Add(PixelStreaming2StatNames::PliCount, 0);
+		Add(PixelStreaming2StatNames::NackCount, 0);
+		Add(PixelStreaming2StatNames::RetransmittedBytesSent, 0);
+		Add(PixelStreaming2StatNames::TotalEncodeBytesTarget, 0);
+		Add(PixelStreaming2StatNames::KeyFramesEncoded, 0);
+		Add(PixelStreaming2StatNames::FrameWidth, 0);
+		Add(PixelStreaming2StatNames::FrameHeight, 0);
+		Add(PixelStreaming2StatNames::HugeFramesSent, 0);
+		Add(PixelStreaming2StatNames::PacketsLost, 0);
+		Add(PixelStreaming2StatNames::Jitter, 0);
+		Add(PixelStreaming2StatNames::RoundTripTime, 0);
+
+		// These are values used to calculate extra values (stores time deltas etc)
+		AddNonRendered(PixelStreaming2StatNames::TargetBitrate);
+		AddNonRendered(PixelStreaming2StatNames::FramesSent);
+		AddNonRendered(PixelStreaming2StatNames::FramesReceived);
+		AddNonRendered(PixelStreaming2StatNames::BytesSent);
+		AddNonRendered(PixelStreaming2StatNames::BytesReceived);
+		AddNonRendered(PixelStreaming2StatNames::QPSum);
+		AddNonRendered(PixelStreaming2StatNames::TotalEncodeTime);
+		AddNonRendered(PixelStreaming2StatNames::FramesEncoded);
+		AddNonRendered(PixelStreaming2StatNames::FramesDecoded);
+		AddNonRendered(PixelStreaming2StatNames::TotalPacketSendDelay);
+
+		// Calculated stats below:
+		// FrameSent Per Second
+		AddStatCalculator([](FStatsSink& StatSource, double Period) -> TOptional<FStatData> {
+			FRTCTrackedStat* FramesSentStat = StatSource.Get(PixelStreaming2StatNames::FramesSent);
+			if (FramesSentStat && FramesSentStat->GetLatestStat().StatValue > 0)
+			{
+				const double FramesSentPerSecond = FramesSentStat->CalculateDelta(Period);
+				FStatData	 FpsStat = FStatData(PixelStreaming2StatNames::FramesSentPerSecond, FramesSentPerSecond, 0);
+				FpsStat.DisplayFlags = FStatData::EDisplayFlags::TEXT | FStatData::EDisplayFlags::GRAPH;
+				return FpsStat;
+			}
+			return {};
+		});
+
+		// FramesReceived Per Second
+		AddStatCalculator([](FStatsSink& StatSource, double Period) -> TOptional<FStatData> {
+			FRTCTrackedStat* FramesReceivedStat = StatSource.Get(PixelStreaming2StatNames::FramesReceived);
+			if (FramesReceivedStat && FramesReceivedStat->GetLatestStat().StatValue > 0)
+			{
+				const double FramesReceivedPerSecond = FramesReceivedStat->CalculateDelta(Period);
+				return FStatData(PixelStreaming2StatNames::FramesReceivedPerSecond, FramesReceivedPerSecond, 0);
+			}
+			return {};
+		});
+
+		// Megabits sent Per Second
+		AddStatCalculator([](FStatsSink& StatSource, double Period) -> TOptional<FStatData> {
+			FRTCTrackedStat* BytesSentStat = StatSource.Get(PixelStreaming2StatNames::BytesSent);
+			if (BytesSentStat && BytesSentStat->GetLatestStat().StatValue > 0)
+			{
+				const double BytesSentPerSecond = BytesSentStat->CalculateDelta(Period);
+				const double MegabitsPerSecond = BytesSentPerSecond / 1'000'000.0 * 8.0;
+				return FStatData(PixelStreaming2StatNames::BitrateMegabits, MegabitsPerSecond, 2);
+			}
+			return {};
+		});
+
+		// Bits sent Per Second
+		AddStatCalculator([](FStatsSink& StatSource, double Period) -> TOptional<FStatData> {
+			FRTCTrackedStat* BytesSentStat = StatSource.Get(PixelStreaming2StatNames::BytesSent);
+			if (BytesSentStat && BytesSentStat->GetLatestStat().StatValue > 0)
+			{
+				const double BytesSentPerSecond = BytesSentStat->CalculateDelta(Period);
+				const double BitsPerSecond = BytesSentPerSecond * 8.0;
+				FStatData	 Stat = FStatData(PixelStreaming2StatNames::Bitrate, BitsPerSecond, 0);
+				Stat.DisplayFlags = FStatData::EDisplayFlags::HIDDEN; // We don't want to display bits per second (too many digits)
+				return Stat;
+			}
+			return {};
+		});
+
+		// Target megabits sent Per Second
+		AddStatCalculator([](FStatsSink& StatSource, double Period) -> TOptional<FStatData> {
+			FRTCTrackedStat* TargetBpsStats = StatSource.Get(PixelStreaming2StatNames::TargetBitrate);
+			if (TargetBpsStats && TargetBpsStats->GetLatestStat().StatValue > 0)
+			{
+				const double TargetBps = TargetBpsStats->Average();
+				const double MegabitsPerSecond = TargetBps / 1'000'000.0;
+				return FStatData(PixelStreaming2StatNames::TargetBitrateMegabits, MegabitsPerSecond, 2);
+			}
+			return {};
+		});
+
+		// Megabits received Per Second
+		AddStatCalculator([](FStatsSink& StatSource, double Period) -> TOptional<FStatData> {
+			FRTCTrackedStat* BytesReceivedStat = StatSource.Get(PixelStreaming2StatNames::BytesReceived);
+			if (BytesReceivedStat && BytesReceivedStat->GetLatestStat().StatValue > 0)
+			{
+				const double BytesReceivedPerSecond = BytesReceivedStat->CalculateDelta(Period);
+				const double MegabitsPerSecond = BytesReceivedPerSecond / 1000.0 * 8.0;
+				return FStatData(PixelStreaming2StatNames::Bitrate, MegabitsPerSecond, 2);
+			}
+			return {};
+		});
 
-		for (const webrtc::RTCStats& Stats : *Report)
-		{
-			for (TSharedPtr<FStatsSink> Sink : RTCStatSinks)
+		// Encoded fps
+		AddStatCalculator([](FStatsSink& StatSource, double Period) -> TOptional<FStatData> {
+			FRTCTrackedStat* EncodedFramesStat = StatSource.Get(PixelStreaming2StatNames::FramesEncoded);
+			if (EncodedFramesStat && EncodedFramesStat->GetLatestStat().StatValue > 0)
 			{
-				if (Sink && Sink->Wants(Stats))
-				{
-					FString PeerId = Sink->DerivePeerId(Stats, AssociatedPlayerId);
-					Sink->Process(Stats, PeerId);
-					Sink->PostProcess(PeerId, SecondsDelta);
-				}
+				const double EncodedFramesPerSecond = EncodedFramesStat->CalculateDelta(Period);
+				return FStatData(PixelStreaming2StatNames::EncodedFramesPerSecond, EncodedFramesPerSecond, 0);
 			}
+			return {};
+		});
 
-			// For debugging to see all stats in the log
-			// UE_LOG(LogPixelStreaming2, Log, TEXT("------------%s---%s------------"), *FString(Stats.id().c_str()), *FString(Stats.type()));
-			// for (const webrtc::RTCStatsMemberInterface* StatMember : StatMembers)
-			// {
-			// 	UE_LOG(LogPixelStreaming2, Log, TEXT("%s"), *FString(StatMember->name()));
-			// }
-		}
-
-		LastCalculationCycles = FPlatformTime::Cycles64();
-	}
-
-	/*
-	* ---------------- FStatsSink --------------------
-	*/
-
-	void FRTCStatsCollector::FStatsSink::Process(const webrtc::RTCStats& InStats, FString PeerId)
-	{
-		FStats* PSStats = FStats::Get();
-
-		if (!PSStats)
-		{
-			return;
-		}
+		// Decoded fps
+		AddStatCalculator([](FStatsSink& StatSource, double Period) -> TOptional<FStatData> {
+			FRTCTrackedStat* DecodedFramesStat = StatSource.Get(PixelStreaming2StatNames::FramesDecoded);
+			if (DecodedFramesStat && DecodedFramesStat->GetLatestStat().StatValue > 0)
+			{
+				const double DecodedFramesPerSecond = DecodedFramesStat->CalculateDelta(Period);
+				return FStatData(PixelStreaming2StatNames::DecodedFramesPerSecond, DecodedFramesPerSecond, 0);
+			}
+			return {};
+		});
 
-		std::vector<const webrtc::RTCStatsMemberInterface*> StatMembers = InStats.Members();
+		// Avg QP Per Second
+		AddStatCalculator([](FStatsSink& StatSource, double Period) -> TOptional<FStatData> {
+			FRTCTrackedStat* QPSumStat = StatSource.Get(PixelStreaming2StatNames::QPSum);
+			FStatData*		 EncodedFramesPerSecond = StatSource.GetCalculatedStat(PixelStreaming2StatNames::EncodedFramesPerSecond);
+			if (QPSumStat && QPSumStat->GetLatestStat().StatValue > 0
+				&& EncodedFramesPerSecond && EncodedFramesPerSecond->StatValue > 0.0)
+			{
+				const double QPSumDeltaPerSecond = QPSumStat->CalculateDelta(Period);
+				const double MeanQPPerFrame = QPSumDeltaPerSecond / EncodedFramesPerSecond->StatValue;
+				FName		 StatName = PixelStreaming2StatNames::MeanQPPerSecond;
+				return FStatData(StatName, MeanQPPerFrame, 0);
+			}
+			return {};
+		});
 
-		for (const webrtc::RTCStatsMemberInterface* StatMember : StatMembers)
-		{
-			const FName StatName = FName(StatMember->name());
+		// Mean EncodeTime (ms) Per Frame
+		AddStatCalculator([](FStatsSink& StatSource, double Period) -> TOptional<FStatData> {
+			FRTCTrackedStat* TotalEncodeTimeStat = StatSource.Get(PixelStreaming2StatNames::TotalEncodeTime);
+			FStatData*		 EncodedFramesPerSecond = StatSource.GetCalculatedStat(PixelStreaming2StatNames::EncodedFramesPerSecond);
+			if (TotalEncodeTimeStat && TotalEncodeTimeStat->GetLatestStat().StatValue > 0
+				&& EncodedFramesPerSecond && EncodedFramesPerSecond->StatValue > 0.0)
+			{
+				const double TotalEncodeTimePerSecond = TotalEncodeTimeStat->CalculateDelta(Period);
+				const double MeanEncodeTimePerFrameMs = TotalEncodeTimePerSecond / EncodedFramesPerSecond->StatValue * 1000.0;
+				return FStatData(PixelStreaming2StatNames::MeanEncodeTime, MeanEncodeTimePerFrameMs, 2);
+			}
+			return {};
+		});
 
-			FRTCTrackedStat* StatToEmit = Get(StatName);
-			if (!StatToEmit)
+		// Mean SendDelay (ms) Per Frame
+		AddStatCalculator([](FStatsSink& StatSource, double Period) -> TOptional<FStatData> {
+			FRTCTrackedStat* TotalSendDelayStat = StatSource.Get(PixelStreaming2StatNames::TotalPacketSendDelay);
+			FStatData*		 FramesSentPerSecond = StatSource.GetCalculatedStat(PixelStreaming2StatNames::FramesSentPerSecond);
+			if (TotalSendDelayStat && TotalSendDelayStat->GetLatestStat().StatValue > 0
+				&& FramesSentPerSecond && FramesSentPerSecond->StatValue > 0.0)
 			{
-				continue;
+				const double TotalSendDelayPerSecond = TotalSendDelayStat->CalculateDelta(Period);
+				const double MeanSendDelayPerFrameMs = TotalSendDelayPerSecond / FramesSentPerSecond->StatValue * 1000.0;
+				return FStatData(PixelStreaming2StatNames::MeanSendDelay, MeanSendDelayPerFrameMs, 2);
 			}
+			return {};
+		});
 
-			if (ExtractValueAndSet(StatMember, StatToEmit))
+		// JitterBufferDelay (ms)
+		AddStatCalculator([](FStatsSink& StatSource, double Period) -> TOptional<FStatData> {
+			FRTCTrackedStat* JitterBufferDelayStat = StatSource.Get(PixelStreaming2StatNames::JitterBufferDelay);
+			FStatData*		 FramesReceivedPerSecond = StatSource.GetCalculatedStat(PixelStreaming2StatNames::FramesReceivedPerSecond);
+			if (JitterBufferDelayStat && JitterBufferDelayStat->GetLatestStat().StatValue > 0
+				&& FramesReceivedPerSecond && FramesReceivedPerSecond->StatValue > 0.0)
 			{
-				PSStats->StorePeerStat(PeerId, SinkType, StatToEmit->GetLatestStat());
+				const double TotalJitterBufferDelayPerSecond = JitterBufferDelayStat->CalculateDelta(Period);
+				const double MeanJitterBufferDelayMs = TotalJitterBufferDelayPerSecond / FramesReceivedPerSecond->StatValue * 1000.0;
+				return FStatData(PixelStreaming2StatNames::JitterBufferDelay, MeanJitterBufferDelayMs, 2);
 			}
-		}
+			return {};
+		});
 	}
 
-	void FRTCStatsCollector::FStatsSink::PostProcess(FString PeerId, double SecondsDelta)
+	void FRTCStatsCollector::FRTPLocalAudioTrackStatsSink::Process(const EpicRtcLocalTrackRtpStats& InStats, const FString& PeerId, double SecondsDelta)
 	{
 		FStats* PSStats = FStats::Get();
-
 		if (!PSStats)
 		{
 			return;
 		}
 
-		// Run all the stat calculators
-		for (auto& Calculator : Calculators)
+		for (TPair<FName, FRTCTrackedStat>& Tuple : Stats)
 		{
-			TOptional<FStatData> OptStatData = Calculator(*this, SecondsDelta);
-			if (OptStatData.IsSet())
+			double NewValue = 0;
+			if (Tuple.Key == PixelStreaming2StatNames::TotalPacketSendDelay)
 			{
-				FStatData& StatData = *OptStatData;
-				CalculatedStats.Add(StatData.StatName, StatData);
-				PSStats->StorePeerStat(PeerId, SinkType, StatData);
+				NewValue = InStats._local._totalPacketSendDelay;
+			}
+			else if (Tuple.Key == PixelStreaming2StatNames::TargetBitrate)
+			{
+				NewValue = InStats._local._targetBitrate;
+			}
+			else if (Tuple.Key == PixelStreaming2StatNames::BytesSent)
+			{
+				NewValue = InStats._local._bytesSent;
+			}
+			else if (Tuple.Key == PixelStreaming2StatNames::PacketsLost)
+			{
+				NewValue = InStats._remote._packetsLost;
+			}
+			else if (Tuple.Key == PixelStreaming2StatNames::Jitter)
+			{
+				NewValue = InStats._remote._jitter;
+			}
+			else if (Tuple.Key == PixelStreaming2StatNames::RoundTripTime)
+			{
+				NewValue = InStats._remote._roundTripTime;
 			}
-		}
-	}
 
-	bool FRTCStatsCollector::FStatsSink::ExtractValueAndSet(const webrtc::RTCStatsMemberInterface* ExtractFrom, FRTCTrackedStat* SetValueHere)
-	{
-		const bool bIsDefined = ExtractFrom->is_defined();
-		if (!bIsDefined)
-		{
-			return false;
+			if (UpdateValue(NewValue, &Tuple.Value))
+			{
+				PSStats->StorePeerStat(PeerId, RTCStatCategories::LocalAudioTrack, Tuple.Value.GetLatestStat());
+			}
 		}
 
-		const bool bZeroInitially = SetValueHere->GetLatestStat().StatValue == 0.0;
-		FString StatValueStr = ToString(ExtractFrom->ValueToString());
-		double StatValueDouble = FCString::Atod(*StatValueStr);
-		SetValueHere->SetLatestValue(StatValueDouble);
-		const bool bZeroStill = SetValueHere->GetLatestStat().StatValue == 0.0;
-		return !(bZeroInitially && bZeroStill);
-	}
-
-	void FRTCStatsCollector::FStatsSink::AddAliased(FName StatName, FName AliasedName, int NDecimalPlaces, uint8 DisplayFlags)
-	{
-		FRTCTrackedStat Stat = FRTCTrackedStat(StatName, AliasedName, NDecimalPlaces, DisplayFlags);
-		Stats.Add(StatName, Stat);
+		PostProcess(PSStats, PeerId, SecondsDelta);
 	}
 
 	/*
-	* ---------------- FRTPMediaStatsSink --------------------
-	*/
-
-	FRTCStatsCollector::FRTPMediaStatsSink::FRTPMediaStatsSink(FName InSinkType, FString InMediaKind)
-		: FStatsSink(InSinkType)
-		, MediaKind(InMediaKind)
+	 * ---------- FRTPRemoteTrackStatsSink ----------
+	 */
+	FRTCStatsCollector::FRTPRemoteTrackStatsSink::FRTPRemoteTrackStatsSink(FName InCategory)
+		: FStatsSink(InCategory)
 	{
-		// Todo: These stats could be split up across more specific classes once we want start tracking audio stats.
-
 		// These stats will be extracted from the stat reports and emitted straight to screen
 		Add(PixelStreaming2StatNames::FirCount, 0);
 		Add(PixelStreaming2StatNames::PliCount, 0);
 		Add(PixelStreaming2StatNames::NackCount, 0);
-		Add(PixelStreaming2StatNames::SliCount, 0);
-		Add(PixelStreaming2StatNames::RetransmittedBytesSent, 0);
+		Add(PixelStreaming2StatNames::RetransmittedBytesReceived, 0);
+		Add(PixelStreaming2StatNames::RetransmittedPacketsReceived, 0);
 		Add(PixelStreaming2StatNames::TotalEncodeBytesTarget, 0);
-		Add(PixelStreaming2StatNames::KeyFramesEncoded, 0);
+		Add(PixelStreaming2StatNames::KeyFramesDecoded, 0);
 		Add(PixelStreaming2StatNames::FrameWidth, 0);
 		Add(PixelStreaming2StatNames::FrameHeight, 0);
 		Add(PixelStreaming2StatNames::HugeFramesSent, 0);
-		Add(PixelStreaming2StatNames::AvgSendDelay, 0);
+		Add(PixelStreaming2StatNames::PacketsLost, 0);
+		Add(PixelStreaming2StatNames::Jitter, 0);
+		Add(PixelStreaming2StatNames::RoundTripTime, 0);
 
 		// These are values used to calculate extra values (stores time deltas etc)
 		AddNonRendered(PixelStreaming2StatNames::TargetBitrate);
@@ -193,21 +694,21 @@ namespace UE::PixelStreaming2
 		AddNonRendered(PixelStreaming2StatNames::TotalEncodeTime);
 		AddNonRendered(PixelStreaming2StatNames::FramesEncoded);
 		AddNonRendered(PixelStreaming2StatNames::FramesDecoded);
+		AddNonRendered(PixelStreaming2StatNames::TotalPacketSendDelay);
 
 		// Calculated stats below:
-
 		// FrameSent Per Second
 		AddStatCalculator([](FStatsSink& StatSource, double Period) -> TOptional<FStatData> {
 			FRTCTrackedStat* FramesSentStat = StatSource.Get(PixelStreaming2StatNames::FramesSent);
 			if (FramesSentStat && FramesSentStat->GetLatestStat().StatValue > 0)
 			{
 				const double FramesSentPerSecond = FramesSentStat->CalculateDelta(Period);
-				FStatData FpsStat = FStatData(PixelStreaming2StatNames::FramesSentPerSecond, FramesSentPerSecond, 0);
+				FStatData	 FpsStat = FStatData(PixelStreaming2StatNames::FramesSentPerSecond, FramesSentPerSecond, 0);
 				FpsStat.DisplayFlags = FStatData::EDisplayFlags::TEXT | FStatData::EDisplayFlags::GRAPH;
 				return FpsStat;
 			}
 			return {};
-			});
+		});
 
 		// FramesReceived Per Second
 		AddStatCalculator([](FStatsSink& StatSource, double Period) -> TOptional<FStatData> {
@@ -218,7 +719,7 @@ namespace UE::PixelStreaming2
 				return FStatData(PixelStreaming2StatNames::FramesReceivedPerSecond, FramesReceivedPerSecond, 0);
 			}
 			return {};
-			});
+		});
 
 		// Megabits sent Per Second
 		AddStatCalculator([](FStatsSink& StatSource, double Period) -> TOptional<FStatData> {
@@ -230,7 +731,7 @@ namespace UE::PixelStreaming2
 				return FStatData(PixelStreaming2StatNames::BitrateMegabits, MegabitsPerSecond, 2);
 			}
 			return {};
-			});
+		});
 
 		// Bits sent Per Second
 		AddStatCalculator([](FStatsSink& StatSource, double Period) -> TOptional<FStatData> {
@@ -239,12 +740,12 @@ namespace UE::PixelStreaming2
 			{
 				const double BytesSentPerSecond = BytesSentStat->CalculateDelta(Period);
 				const double BitsPerSecond = BytesSentPerSecond * 8.0;
-				FStatData Stat = FStatData(PixelStreaming2StatNames::Bitrate, BitsPerSecond, 0);
+				FStatData	 Stat = FStatData(PixelStreaming2StatNames::Bitrate, BitsPerSecond, 0);
 				Stat.DisplayFlags = FStatData::EDisplayFlags::HIDDEN; // We don't want to display bits per second (too many digits)
 				return Stat;
 			}
 			return {};
-			});
+		});
 
 		// Target megabits sent Per Second
 		AddStatCalculator([](FStatsSink& StatSource, double Period) -> TOptional<FStatData> {
@@ -256,7 +757,7 @@ namespace UE::PixelStreaming2
 				return FStatData(PixelStreaming2StatNames::TargetBitrateMegabits, MegabitsPerSecond, 2);
 			}
 			return {};
-			});
+		});
 
 		// Megabits received Per Second
 		AddStatCalculator([](FStatsSink& StatSource, double Period) -> TOptional<FStatData> {
@@ -268,7 +769,7 @@ namespace UE::PixelStreaming2
 				return FStatData(PixelStreaming2StatNames::Bitrate, MegabitsPerSecond, 2);
 			}
 			return {};
-			});
+		});
 
 		// Encoded fps
 		AddStatCalculator([](FStatsSink& StatSource, double Period) -> TOptional<FStatData> {
@@ -279,7 +780,7 @@ namespace UE::PixelStreaming2
 				return FStatData(PixelStreaming2StatNames::EncodedFramesPerSecond, EncodedFramesPerSecond, 0);
 			}
 			return {};
-			});
+		});
 
 		// Decoded fps
 		AddStatCalculator([](FStatsSink& StatSource, double Period) -> TOptional<FStatData> {
@@ -290,27 +791,27 @@ namespace UE::PixelStreaming2
 				return FStatData(PixelStreaming2StatNames::DecodedFramesPerSecond, DecodedFramesPerSecond, 0);
 			}
 			return {};
-			});
+		});
 
 		// Avg QP Per Second
 		AddStatCalculator([](FStatsSink& StatSource, double Period) -> TOptional<FStatData> {
 			FRTCTrackedStat* QPSumStat = StatSource.Get(PixelStreaming2StatNames::QPSum);
-			FStatData* EncodedFramesPerSecond = StatSource.GetCalculatedStat(PixelStreaming2StatNames::EncodedFramesPerSecond);
+			FStatData*		 EncodedFramesPerSecond = StatSource.GetCalculatedStat(PixelStreaming2StatNames::EncodedFramesPerSecond);
 			if (QPSumStat && QPSumStat->GetLatestStat().StatValue > 0
 				&& EncodedFramesPerSecond && EncodedFramesPerSecond->StatValue > 0.0)
 			{
 				const double QPSumDeltaPerSecond = QPSumStat->CalculateDelta(Period);
 				const double MeanQPPerFrame = QPSumDeltaPerSecond / EncodedFramesPerSecond->StatValue;
-				FName StatName = PixelStreaming2StatNames::MeanQPPerSecond;
+				FName		 StatName = PixelStreaming2StatNames::MeanQPPerSecond;
 				return FStatData(StatName, MeanQPPerFrame, 0);
 			}
 			return {};
-			});
+		});
 
 		// Mean EncodeTime (ms) Per Frame
 		AddStatCalculator([](FStatsSink& StatSource, double Period) -> TOptional<FStatData> {
 			FRTCTrackedStat* TotalEncodeTimeStat = StatSource.Get(PixelStreaming2StatNames::TotalEncodeTime);
-			FStatData* EncodedFramesPerSecond = StatSource.GetCalculatedStat(PixelStreaming2StatNames::EncodedFramesPerSecond);
+			FStatData*		 EncodedFramesPerSecond = StatSource.GetCalculatedStat(PixelStreaming2StatNames::EncodedFramesPerSecond);
 			if (TotalEncodeTimeStat && TotalEncodeTimeStat->GetLatestStat().StatValue > 0
 				&& EncodedFramesPerSecond && EncodedFramesPerSecond->StatValue > 0.0)
 			{
@@ -319,12 +820,12 @@ namespace UE::PixelStreaming2
 				return FStatData(PixelStreaming2StatNames::MeanEncodeTime, MeanEncodeTimePerFrameMs, 2);
 			}
 			return {};
-			});
+		});
 
 		// Mean SendDelay (ms) Per Frame
 		AddStatCalculator([](FStatsSink& StatSource, double Period) -> TOptional<FStatData> {
 			FRTCTrackedStat* TotalSendDelayStat = StatSource.Get(PixelStreaming2StatNames::TotalPacketSendDelay);
-			FStatData* FramesSentPerSecond = StatSource.GetCalculatedStat(PixelStreaming2StatNames::FramesSentPerSecond);
+			FStatData*		 FramesSentPerSecond = StatSource.GetCalculatedStat(PixelStreaming2StatNames::FramesSentPerSecond);
 			if (TotalSendDelayStat && TotalSendDelayStat->GetLatestStat().StatValue > 0
 				&& FramesSentPerSecond && FramesSentPerSecond->StatValue > 0.0)
 			{
@@ -333,12 +834,12 @@ namespace UE::PixelStreaming2
 				return FStatData(PixelStreaming2StatNames::MeanSendDelay, MeanSendDelayPerFrameMs, 2);
 			}
 			return {};
-			});
+		});
 
 		// JitterBufferDelay (ms)
 		AddStatCalculator([](FStatsSink& StatSource, double Period) -> TOptional<FStatData> {
 			FRTCTrackedStat* JitterBufferDelayStat = StatSource.Get(PixelStreaming2StatNames::JitterBufferDelay);
-			FStatData* FramesReceivedPerSecond = StatSource.GetCalculatedStat(PixelStreaming2StatNames::FramesReceivedPerSecond);
+			FStatData*		 FramesReceivedPerSecond = StatSource.GetCalculatedStat(PixelStreaming2StatNames::FramesReceivedPerSecond);
 			if (JitterBufferDelayStat && JitterBufferDelayStat->GetLatestStat().StatValue > 0
 				&& FramesReceivedPerSecond && FramesReceivedPerSecond->StatValue > 0.0)
 			{
@@ -347,106 +848,256 @@ namespace UE::PixelStreaming2
 				return FStatData(PixelStreaming2StatNames::JitterBufferDelay, MeanJitterBufferDelayMs, 2);
 			}
 			return {};
-			});
+		});
+	}
+
+	void FRTCStatsCollector::FRTPRemoteTrackStatsSink::Process(const EpicRtcRemoteTrackRtpStats& InStats, const FString& PeerId, double SecondsDelta)
+	{
+		FStats* PSStats = FStats::Get();
+		if (!PSStats)
+		{
+			return;
+		}
+
+		for (TPair<FName, FRTCTrackedStat>& Tuple : Stats)
+		{
+			double NewValue = 0;
+			if (Tuple.Key == PixelStreaming2StatNames::FirCount)
+			{
+				NewValue = InStats._local._firCount;
+			}
+			else if (Tuple.Key == PixelStreaming2StatNames::PliCount)
+			{
+				NewValue = InStats._local._pliCount;
+			}
+			else if (Tuple.Key == PixelStreaming2StatNames::NackCount)
+			{
+				NewValue = InStats._local._nackCount;
+			}
+			else if (Tuple.Key == PixelStreaming2StatNames::RetransmittedBytesReceived)
+			{
+				NewValue = InStats._local._retransmittedBytesReceived;
+			}
+			else if (Tuple.Key == PixelStreaming2StatNames::RetransmittedPacketsReceived)
+			{
+				NewValue = InStats._local._retransmittedPacketsReceived;
+			}
+			else if (Tuple.Key == PixelStreaming2StatNames::KeyFramesDecoded)
+			{
+				NewValue = InStats._local._keyFramesDecoded;
+			}
+			else if (Tuple.Key == PixelStreaming2StatNames::FrameWidth)
+			{
+				NewValue = InStats._local._frameWidth;
+			}
+			else if (Tuple.Key == PixelStreaming2StatNames::FrameHeight)
+			{
+				NewValue = InStats._local._frameHeight;
+			}
+			else if (Tuple.Key == PixelStreaming2StatNames::FramesReceived)
+			{
+				NewValue = InStats._local._framesReceived;
+			}
+			else if (Tuple.Key == PixelStreaming2StatNames::BytesReceived)
+			{
+				NewValue = InStats._local._bytesReceived;
+			}
+			else if (Tuple.Key == PixelStreaming2StatNames::QPSum)
+			{
+				NewValue = InStats._local._qpSum;
+			}
+			else if (Tuple.Key == PixelStreaming2StatNames::FramesDecoded)
+			{
+				NewValue = InStats._local._framesDecoded;
+			}
+			else if (Tuple.Key == PixelStreaming2StatNames::PacketsLost)
+			{
+				NewValue = InStats._local._packetsLost;
+			}
+			else if (Tuple.Key == PixelStreaming2StatNames::Jitter)
+			{
+				NewValue = InStats._local._jitter;
+			}
+			else if (Tuple.Key == PixelStreaming2StatNames::RoundTripTime)
+			{
+				NewValue = InStats._remote._roundTripTime;
+			}
+
+			if (UpdateValue(NewValue, &Tuple.Value))
+			{
+				PSStats->StorePeerStat(PeerId, RTCStatCategories::RemoteVideoTrack, Tuple.Value.GetLatestStat());
+			}
+		}
+
+		PostProcess(PSStats, PeerId, SecondsDelta);
 	}
 
-	bool FRTCStatsCollector::FRTPMediaStatsSink::Wants(const webrtc::RTCStats& InStats) const
+	/*
+	 * ---------- FRTPVideoSourceSink ----------
+	 */
+	FRTCStatsCollector::FVideoSourceStatsSink::FVideoSourceStatsSink()
+		: FStatsSink(RTCStatCategories::VideoSource)
 	{
-		bool bWants = FStatsSink::Wants(InStats);
+		// Track video source fps
+		Add(PixelStreaming2StatNames::SourceFps, 0);
+	}
 
-		if (!bWants)
+	void FRTCStatsCollector::FVideoSourceStatsSink::Process(const EpicRtcVideoSourceStats& InStats, const FString& PeerId, double SecondsDelta)
+	{
+		FStats* PSStats = FStats::Get();
+		if (!PSStats)
 		{
-			return false;
+			return;
 		}
-		// Check the `kind` field to see if it matches what we want
-		std::vector<const webrtc::RTCStatsMemberInterface*> StatMembers = InStats.Members();
-		for (const webrtc::RTCStatsMemberInterface* StatMember : StatMembers)
+
+		for (TPair<FName, FRTCTrackedStat>& Tuple : Stats)
 		{
-			const FString StatName = FString(StatMember->name());
-			if (StatName == "kind")
+			double NewValue = 0;
+			if (Tuple.Key == PixelStreaming2StatNames::SourceFps)
 			{
-				FString StatMediaKind = StatMember->is_defined() ? ToString(StatMember->ValueToString()) : TEXT("");
-				return MediaKind == StatMediaKind;
+				NewValue = InStats._framesPerSecond;
+			}
+
+			if (UpdateValue(NewValue, &Tuple.Value))
+			{
+				PSStats->StorePeerStat(PeerId, RTCStatCategories::VideoSource, Tuple.Value.GetLatestStat());
 			}
 		}
-		return false;
+
+		PostProcess(PSStats, PeerId, SecondsDelta);
+	}
+
+	/*
+	 * ---------- FRTPAudioSourceSink ----------
+	 */
+	FRTCStatsCollector::FAudioSourceStatsSink::FAudioSourceStatsSink()
+		: FStatsSink(RTCStatCategories::AudioSource)
+	{
+		Add(PixelStreaming2StatNames::AudioLevel, 0);
+		Add(PixelStreaming2StatNames::TotalSamplesDuration, 0);
 	}
 
-	FString FRTCStatsCollector::FRTPMediaStatsSink::DerivePeerId(const webrtc::RTCStats& InStats, FPixelStreaming2PlayerId PlayerId) const
+	void FRTCStatsCollector::FAudioSourceStatsSink::Process(const EpicRtcAudioSourceStats& InStats, const FString& PeerId, double SecondsDelta)
 	{
-		std::vector<const webrtc::RTCStatsMemberInterface*> StatMembers = InStats.Members();
-		const FString StatsType = FString(InStats.type());
+		FStats* PSStats = FStats::Get();
+		if (!PSStats)
+		{
+			return;
+		}
 
-		if (PlayerId == SFU_PLAYER_ID && (StatsType == RTCStatTypes::OutboundRTP || StatsType == RTCStatTypes::InboundRTP))
+		for (TPair<FName, FRTCTrackedStat>& Tuple : Stats)
 		{
-			// Extract the `ssrc` to uniquely id the stream
-			for (const webrtc::RTCStatsMemberInterface* StatMember : StatMembers)
+			double NewValue = 0;
+			if (Tuple.Key == PixelStreaming2StatNames::AudioLevel)
 			{
-				const FString StatName = FString(StatMember->name());
-				if (StatName == "ssrc")
-				{
-					FString Ssrc = StatMember->is_defined() ? ToString(StatMember->ValueToString()) : TEXT("");
-					return Ssrc;
-				}
+				NewValue = InStats._audioLevel;
+			}
+			else if (Tuple.Key == PixelStreaming2StatNames::TotalSamplesDuration)
+			{
+				NewValue = InStats._totalSamplesDuration;
+			}
+
+			if (UpdateValue(NewValue, &Tuple.Value))
+			{
+				PSStats->StorePeerStat(PeerId, RTCStatCategories::AudioSource, Tuple.Value.GetLatestStat());
 			}
 		}
 
-		return PlayerId;
+		PostProcess(PSStats, PeerId, SecondsDelta);
 	}
 
 	/*
-	* ---------------- FTrackStatsSink ----------------
-	*/
-	FRTCStatsCollector::FTrackStatsSink::FTrackStatsSink() : FStatsSink(RTCStatTypes::Track)
+	 * ----------- FDataChannelSink -----------
+	 */
+	FRTCStatsCollector::FDataTrackStatsSink::FDataTrackStatsSink()
+		: FStatsSink(RTCStatCategories::DataChannel)
 	{
-		// Basic stats we wish to store and emit
-		Add(PixelStreaming2StatNames::JitterBufferDelay, 2);
-		Add(PixelStreaming2StatNames::FramesPerSecond, 0);
-		Add(PixelStreaming2StatNames::FramesDecoded, 0);
-		Add(PixelStreaming2StatNames::FramesDropped, 0);
-		Add(PixelStreaming2StatNames::FramesCorrupted, 0);
-		Add(PixelStreaming2StatNames::PartialFramesLost, 0);
-		Add(PixelStreaming2StatNames::FullFramesLost, 0);
-		Add(PixelStreaming2StatNames::JitterBufferTargetDelay, 2);
-		Add(PixelStreaming2StatNames::InterruptionCount, 0);
-		Add(PixelStreaming2StatNames::TotalInterruptionDuration, 2);
-		Add(PixelStreaming2StatNames::FreezeCount, 0);
-		Add(PixelStreaming2StatNames::PauseCount, 0);
-		Add(PixelStreaming2StatNames::TotalFreezesDuration, 2);
-		Add(PixelStreaming2StatNames::TotalPausesDuration, 2);
+		// These names are added as aliased names because `bytesSent` is ambiguous stat that is used across inbound-rtp, outbound-rtp, and data-channel
+		// so to disambiguate which state we are referring to we record the `bytesSent` stat for the data-channel but store and report it as `data-channel-bytesSent`
+		AddAliased(PixelStreaming2StatNames::MessagesSent, PixelStreaming2StatNames::DataChannelMessagesSent, 0, FStatData::EDisplayFlags::TEXT | FStatData::EDisplayFlags::GRAPH);
+		AddAliased(PixelStreaming2StatNames::MessagesReceived, PixelStreaming2StatNames::DataChannelBytesReceived, 0, FStatData::EDisplayFlags::TEXT | FStatData::EDisplayFlags::GRAPH);
+		AddAliased(PixelStreaming2StatNames::BytesSent, PixelStreaming2StatNames::DataChannelBytesSent, 0, FStatData::EDisplayFlags::TEXT | FStatData::EDisplayFlags::GRAPH);
+		AddAliased(PixelStreaming2StatNames::BytesReceived, PixelStreaming2StatNames::DataChannelMessagesReceived, 0, FStatData::EDisplayFlags::TEXT | FStatData::EDisplayFlags::GRAPH);
 	}
 
+	void FRTCStatsCollector::FDataTrackStatsSink::Process(const EpicRtcDataTrackStats& InStats, const FString& PeerId, double SecondsDelta)
+	{
+		FStats* PSStats = FStats::Get();
+		if (!PSStats)
+		{
+			return;
+		}
+
+		for (TPair<FName, FRTCTrackedStat>& Tuple : Stats)
+		{
+			double NewValue = 0;
+			if (Tuple.Key == PixelStreaming2StatNames::MessagesSent)
+			{
+				NewValue = InStats._messagesSent;
+			}
+			else if (Tuple.Key == PixelStreaming2StatNames::MessagesReceived)
+			{
+				NewValue = InStats._messagesReceived;
+			}
+			else if (Tuple.Key == PixelStreaming2StatNames::BytesSent)
+			{
+				NewValue = InStats._bytesSent;
+			}
+			else if (Tuple.Key == PixelStreaming2StatNames::BytesReceived)
+			{
+				NewValue = InStats._bytesReceived;
+			}
+
+			if (UpdateValue(NewValue, &Tuple.Value))
+			{
+				PSStats->StorePeerStat(PeerId, RTCStatCategories::DataChannel, Tuple.Value.GetLatestStat());
+			}
+		}
+
+		PostProcess(PSStats, PeerId, SecondsDelta);
+	}
 	/*
-	* ---------------- FVideoSourceStatsSink ----------------
-	*/
-	FRTCStatsCollector::FVideoSourceStatsSink::FVideoSourceStatsSink() : FStatsSink(RTCStatTypes::MediaSource)
+	 * ---------- FRTPAudioSourceSink ----------
+	 */
+	FRTCStatsCollector::FCandidatePairStatsSink::FCandidatePairStatsSink()
+		: FStatsSink(RTCStatCategories::CandidatePair)
 	{
-		// Track video source fps
-		Add(PixelStreaming2StatNames::SourceFps, 0);
+		Add(PixelStreaming2StatNames::AvailableOutgoingBitrate, 0);
+		Add(PixelStreaming2StatNames::AvailableIncomingBitrate, 0);
 	}
 
-	bool FRTCStatsCollector::FVideoSourceStatsSink::Wants(const webrtc::RTCStats& InStats) const
+	void FRTCStatsCollector::FCandidatePairStatsSink::Process(const EpicRtcIceCandidatePairStats& InStats, const FString& PeerId, double SecondsDelta)
 	{
-		const FString StatsType = FString(InStats.type());
-		if (StatsType != RTCStatTypes::MediaSource)
+		FStats* PSStats = FStats::Get();
+		if (!PSStats)
 		{
-			return false;
+			return;
 		}
-		for (const webrtc::RTCStatsMemberInterface* StatMember : InStats.Members())
+
+		for (TPair<FName, FRTCTrackedStat>& Tuple : Stats)
 		{
-			const FString StatName = FString(StatMember->name());
-			if (StatName == "kind")
+			double NewValue = 0;
+			if (Tuple.Key == PixelStreaming2StatNames::AvailableOutgoingBitrate)
+			{
+				NewValue = InStats._availableOutgoingBitrate;
+			}
+			else if (Tuple.Key == PixelStreaming2StatNames::AvailableIncomingBitrate)
 			{
-				FString StatMediaKind = StatMember->is_defined() ? ToString(StatMember->ValueToString()) : TEXT("");
-				return StatMediaKind == "video";
+				NewValue = InStats._availableIncomingBitrate;
+			}
+
+			if (UpdateValue(NewValue, &Tuple.Value))
+			{
+				PSStats->StorePeerStat(PeerId, Category, Tuple.Value.GetLatestStat());
 			}
 		}
-		return false;
+
+		PostProcess(PSStats, PeerId, SecondsDelta);
 	}
 
 	/*
-	* ---------- FRTCTrackedStat -------------------
-	*/
+	 * ---------- FRTCTrackedStat -------------------
+	 */
 	FRTCStatsCollector::FRTCTrackedStat::FRTCTrackedStat(FName StatName, FName Alias, int NDecimalPlaces, uint8 DisplayFlags)
 		: LatestStat(StatName, 0.0, NDecimalPlaces)
 	{
@@ -471,18 +1122,38 @@ namespace UE::PixelStreaming2
 	}
 
 	/*
-	* ----------- FDataChannelStatsSink -----------
-	*/
-	FRTCStatsCollector::FDataChannelStatsSink::FDataChannelStatsSink() : FStatsSink("data-channel")
+	 * --------- FStatsSink ------------------------
+	 */
+	FRTCStatsCollector::FStatsSink::FStatsSink(FName InCategory)
+		: Category(MoveTemp(InCategory))
 	{
-		// These names are added as aliased names because `bytesSent` is ambiguous stat that is used across inbound-rtp, outbound-rtp, and data-channel
-		// so to disambiguate which state we are referring to we record the `bytesSent` stat for the data-channel but store and report it as `data-channel-bytesSent`
-		AddAliased(PixelStreaming2StatNames::MessagesSent, PixelStreaming2StatNames::DataChannelMessagesSent, 0, FStatData::EDisplayFlags::TEXT | FStatData::EDisplayFlags::GRAPH);
-		AddAliased(PixelStreaming2StatNames::MessagesReceived, PixelStreaming2StatNames::DataChannelBytesReceived, 0, FStatData::EDisplayFlags::TEXT | FStatData::EDisplayFlags::GRAPH);
-		AddAliased(PixelStreaming2StatNames::BytesSent, PixelStreaming2StatNames::DataChannelBytesSent, 0, FStatData::EDisplayFlags::TEXT | FStatData::EDisplayFlags::GRAPH);
-		AddAliased(PixelStreaming2StatNames::BytesReceived, PixelStreaming2StatNames::DataChannelMessagesReceived, 0, FStatData::EDisplayFlags::TEXT | FStatData::EDisplayFlags::GRAPH);
 	}
 
-} // namespace UE::PixelStreaming2
+	void FRTCStatsCollector::FStatsSink::AddAliased(FName StatName, FName AliasedName, int NDecimalPlaces, uint8 DisplayFlags)
+	{
+		FRTCTrackedStat Stat = FRTCTrackedStat(StatName, AliasedName, NDecimalPlaces, DisplayFlags);
+		Stats.Add(StatName, Stat);
+	}
+
+	bool FRTCStatsCollector::FStatsSink::UpdateValue(double NewValue, FRTCTrackedStat* SetValueHere)
+	{
+		const bool bZeroInitially = SetValueHere->GetLatestStat().StatValue == 0.0;
+		SetValueHere->SetLatestValue(NewValue);
+		const bool bZeroStill = SetValueHere->GetLatestStat().StatValue == 0.0;
+		return !(bZeroInitially && bZeroStill);
+	}
 
-#endif // (Migration)
+	void FRTCStatsCollector::FStatsSink::PostProcess(FStats* PSStats, const FString& PeerId, double SecondsDelta)
+	{
+		for (auto& Calculator : Calculators)
+		{
+			TOptional<FStatData> OptStatData = Calculator(*this, SecondsDelta);
+			if (OptStatData.IsSet())
+			{
+				FStatData& StatData = *OptStatData;
+				CalculatedStats.Add(StatData.StatName, StatData);
+				PSStats->StorePeerStat(PeerId, Category, StatData);
+			}
+		}
+	}
+} // namespace UE::PixelStreaming2
\ No newline at end of file
diff --git a/Engine/Plugins/Media/PixelStreaming2/Source/PixelStreaming2/Private/RTCStatsCollector.h b/Engine/Plugins/Media/PixelStreaming2/Source/PixelStreaming2/Private/RTCStatsCollector.h
--- a/Engine/Plugins/Media/PixelStreaming2/Source/PixelStreaming2/Private/RTCStatsCollector.h
+++ b/Engine/Plugins/Media/PixelStreaming2/Source/PixelStreaming2/Private/RTCStatsCollector.h
@@ -3,148 +3,185 @@
 #pragma once
 
 #include "Stats.h"
-#include "Templates/SharedPointer.h"
+
+#include "epic_rtc/core/stats.h"
 
 namespace UE::PixelStreaming2
 {
-	namespace RTCStatTypes
+	namespace RTCStatCategories
 	{
+		const FName LocalVideoTrack = FName(TEXT("outbound-video-track"));
+		const FName LocalAudioTrack = FName(TEXT("outbound-audio-track"));
+		const FName VideoSource = FName(TEXT("video-source"));
+		const FName AudioSource = FName(TEXT("audio-source"));
 		const FName DataChannel = FName(TEXT("data-channel"));
-		const FName OutboundRTP = FName(TEXT("outbound-rtp"));
-		const FName InboundRTP = FName(TEXT("inbound-rtp"));
-		const FName Track = FName(TEXT("track"));
-		const FName MediaSource = FName(TEXT("media-source"));
-	} // namespace RTCStatTypes
-
-	// TODO (Migration) RTCP-6453 Retrieve stats from EpicRtc when EpicRtc supports it
-#if 0
-	class PIXELSTREAMING2_API FRTCStatsCollector : public webrtc::RTCStatsCollectorCallback
+		const FName RemoteVideoTrack = FName(TEXT("inbound-video-track"));
+		const FName RemoteAudioTrack = FName(TEXT("inbound-audio-track"));
+		const FName CandidatePair = FName(TEXT("candidate-pair"));
+	} // namespace RTCStatCategories
+
+	class FRTCStatsCollector
 	{
 	public:
+		static TSharedPtr<FRTCStatsCollector> Create(const FString& PlayerId);
+
+	private:
 		FRTCStatsCollector();
-		FRTCStatsCollector(FString PlayerId);
+		explicit FRTCStatsCollector(const FString& PlayerId);
 
-		// Begin RTCStatsCollectorCallback interface
-		void OnStatsDelivered(const rtc::scoped_refptr<const webrtc::RTCStatsReport>& report) override;
-		void AddRef() const override;
-		rtc::RefCountReleaseStatus Release() const override;
-		// End RTCStatsCollectorCallback interface
+		void OnWebRTCDisableStatsChanged(IConsoleVariable* Var);
+
+	public:
+		void Process(const EpicRtcConnectionStats& InStats);
 
 	private:
 		/*
-		* ---------- FRTCTrackedStat -------------------
-		* Tracks a stat from WebRTC. Stores the current value and previous value.
-		*/
+		 * ---------- FRTCTrackedStat -------------------
+		 * Tracks a stat from WebRTC. Stores the current value and previous value.
+		 */
 		class FRTCTrackedStat
 		{
 		public:
-			FRTCTrackedStat(FName StatName, int NDecimalPlaces) : LatestStat(StatName, 0.0, NDecimalPlaces){}
-			FRTCTrackedStat(FName StatName, int NDecimalPlaces, uint8 DisplayFlags) : LatestStat(StatName, 0.0, NDecimalPlaces){ LatestStat.DisplayFlags = DisplayFlags; }
+			FRTCTrackedStat(FName StatName, int NDecimalPlaces)
+				: LatestStat(StatName, 0.0, NDecimalPlaces)
+			{
+			}
+
+			FRTCTrackedStat(FName StatName, int NDecimalPlaces, uint8 DisplayFlags)
+				: LatestStat(StatName, 0.0, NDecimalPlaces)
+			{
+				LatestStat.DisplayFlags = DisplayFlags;
+			}
+
 			FRTCTrackedStat(FName StatName, FName Alias, int NDecimalPlaces, uint8 DisplayFlags);
-			double CalculateDelta(double Period) const;
-			double Average() const;
+
+			double			 CalculateDelta(double Period) const;
+			double			 Average() const;
 			const FStatData& GetLatestStat() const { return LatestStat; }
-			void SetLatestValue(double InValue);
+			void			 SetLatestValue(double InValue);
+
 		private:
 			FStatData LatestStat;
-			double PrevValue = 0.0;
+			double	  PrevValue = 0.0;
 		};
 
-		/*
-		* ---------- FStatsSink -------------------
-		* A sink for processing webrtc::RTCStats - processing will include extracting stats and storing them.  
-		*/
 		class FStatsSink
 		{
 		public:
-			FStatsSink(FName InSinkType) : SinkType(InSinkType) {}
+			FStatsSink(FName InCategory);
 			virtual ~FStatsSink() = default;
-			virtual bool Wants(const webrtc::RTCStats& InStats) const { return FString(InStats.type()) == SinkType.ToString(); };
-			virtual void Process(const webrtc::RTCStats& InStats, FString PeerId);
-			virtual void PostProcess(FString PeerId, double SecondsDelta);
-			virtual FString DerivePeerId(const webrtc::RTCStats& InStats, FString PlayerId) const { return PlayerId; }
-			// returns true if the value is worth storing (false if it started and remains zero)
-			virtual bool ExtractValueAndSet(const webrtc::RTCStatsMemberInterface* ExtractFrom, FRTCTrackedStat* SetValueHere);
-			virtual bool Contains(const FName StatName) const { return Stats.Contains(StatName); };
-			virtual bool ContainsCalculatedStat(const FName& StatName) const { return CalculatedStats.Contains(StatName); }
-			virtual void Add(FName StatName, int NDecimalPlaces){ Stats.Add(StatName, FRTCTrackedStat(StatName, NDecimalPlaces, FStatData::EDisplayFlags::TEXT)); }
+			virtual void Add(FName StatName, int NDecimalPlaces)
+			{
+				Stats.Add(StatName, FRTCTrackedStat(StatName, NDecimalPlaces, FStatData::EDisplayFlags::TEXT));
+			}
+
 			virtual void AddAliased(FName StatName, FName AliasedName, int NDecimalPlaces, uint8 DisplayFlags);
-			virtual void AddNonRendered(FName StatName) { Stats.Add(StatName, FRTCTrackedStat(StatName, 2, FStatData::EDisplayFlags::HIDDEN)); }
-			virtual void AddStatCalculator(const TFunction<TOptional<FStatData>(FStatsSink&, double)>& Calculator) { Calculators.Add(Calculator); }
-			virtual FRTCTrackedStat* Get(const FName& StatName) { return Stats.Find(StatName); };
-			virtual FStatData* GetCalculatedStat(const FName& StatName) { return CalculatedStats.Find(StatName); }
+			virtual void AddNonRendered(FName StatName)
+			{
+				Stats.Add(StatName, FRTCTrackedStat(StatName, 2, FStatData::EDisplayFlags::HIDDEN));
+			}
+			virtual void AddStatCalculator(const TFunction<TOptional<FStatData>(FStatsSink&, double)>& Calculator)
+			{
+				Calculators.Add(Calculator);
+			}
+			// returns true if the value is worth storing (false if it started and remains zero)
+			virtual bool			 UpdateValue(double NewValue, FRTCTrackedStat* SetValueHere);
+			virtual FRTCTrackedStat* Get(const FName& StatName)
+			{
+				return Stats.Find(StatName);
+			}
+			virtual FStatData* GetCalculatedStat(const FName& StatName)
+			{
+				return CalculatedStats.Find(StatName);
+			}
+			virtual void PostProcess(FStats* PSStats, const FString& PeerId, double SecondsDelta);
 
 			// Stats that are stored as is.
 			TMap<FName, FRTCTrackedStat> Stats;
-
 			// Stats we calculate based on the stats map above. This calculation is done in FStatsSink::PostProcess by the `Calculators` below.
-			TMap<FName, FStatData> CalculatedStats;
+			TMap<FName, FStatData>										 CalculatedStats;
 			TArray<TFunction<TOptional<FStatData>(FStatsSink&, double)>> Calculators;
+
 		protected:
-			FName SinkType;
+			FName Category;
 		};
 
-		/*
-		* ---------- FRTPMediaStatsSink -------------------
-		* Sink useful for RTP media (audio or video) that is inbound or outbound.
-		*/
-		class FRTPMediaStatsSink : public FStatsSink
+		class FRTPLocalVideoTrackStatsSink : public FStatsSink
 		{
 		public:
-			FRTPMediaStatsSink(FName InSinkType, FString InMediaKind);
-			virtual ~FRTPMediaStatsSink() = default;
-			virtual bool Wants(const webrtc::RTCStats& InStats) const override;
-			virtual FString DerivePeerId(const webrtc::RTCStats& InStats, FString PlayerId) const override;
-		private:
-			FString MediaKind;
+			FRTPLocalVideoTrackStatsSink();
+			virtual ~FRTPLocalVideoTrackStatsSink() = default;
+			void Process(const EpicRtcLocalVideoTrackStats& InStats, const FString& PeerId, double SecondsDelta);
 		};
 
-		/*
-		* ---------- FTrackStatsSink -------------------
-		* Stats sink for `track` stats
-		*/
-		class FTrackStatsSink : public FStatsSink
+		class FRTPLocalAudioTrackStatsSink : public FStatsSink
 		{
 		public:
-			FTrackStatsSink();
-			virtual ~FTrackStatsSink() = default;
+			FRTPLocalAudioTrackStatsSink();
+			virtual ~FRTPLocalAudioTrackStatsSink() = default;
+			void Process(const EpicRtcLocalTrackRtpStats& InStats, const FString& PeerId, double SecondsDelta);
+		};
+
+		class FRTPRemoteTrackStatsSink : public FStatsSink
+		{
+		public:
+			FRTPRemoteTrackStatsSink(FName Category);
+			virtual ~FRTPRemoteTrackStatsSink() = default;
+			void Process(const EpicRtcRemoteTrackRtpStats& InStats, const FString& PeerId, double SecondsDelta);
 		};
 
-		/*
-		* ---------- FVideoSourceStatsSink -------------------
-		* Stats sink for video `media-source` stats
-		*/
 		class FVideoSourceStatsSink : public FStatsSink
 		{
 		public:
 			FVideoSourceStatsSink();
 			virtual ~FVideoSourceStatsSink() = default;
-			virtual bool Wants(const webrtc::RTCStats& InStats) const override;
+			void Process(const EpicRtcVideoSourceStats& InStats, const FString& PeerId, double SecondsDelta);
+		};
+
+		class FAudioSourceStatsSink : public FStatsSink
+		{
+		public:
+			FAudioSourceStatsSink();
+			virtual ~FAudioSourceStatsSink() = default;
+			void Process(const EpicRtcAudioSourceStats& InStats, const FString& PeerId, double SecondsDelta);
 		};
 
 		/*
-		* ---------- FDataChannelStatsSink -------------------
-		* Stats sink `data-channel` stats
-		*/
-		class FDataChannelStatsSink : public FStatsSink
+		 * ---------- FDataChannelSink -------------------
+		 */
+		class FDataTrackStatsSink : public FStatsSink
+		{
+		public:
+			FDataTrackStatsSink();
+			virtual ~FDataTrackStatsSink() = default;
+
+			void Process(const EpicRtcDataTrackStats& InStats, const FString& PeerId, double SecondsDelta);
+		};
+
+		class FCandidatePairStatsSink : public FStatsSink
 		{
 		public:
-			FDataChannelStatsSink();
-			virtual ~FDataChannelStatsSink() = default;
+			FCandidatePairStatsSink();
+			virtual ~FCandidatePairStatsSink() = default;
+			void Process(const EpicRtcIceCandidatePairStats& InStats, const FString& PeerId, double SecondsDelta);
 		};
 
 	private:
-		FString AssociatedPlayerId;
+		FString									 AssociatedPlayerId;
+		uint64									 LastCalculationCycles;
+		bool									 bIsEnabled;
 
-		mutable int32 RefCount;
+		TUniquePtr<FRTPLocalVideoTrackStatsSink> LocalVideoTrackStatsSink;
+		TUniquePtr<FRTPLocalAudioTrackStatsSink> LocalAudioTrackStatsSink;
 
-		uint64 LastCalculationCycles;
+		TUniquePtr<FVideoSourceStatsSink>		 VideoSourceStatsSink;
+		TUniquePtr<FAudioSourceStatsSink>	 AudioSourceStatsSink;
 
-		bool bIsEnabled;
+		TUniquePtr<FRTPRemoteTrackStatsSink> RemoteVideoTrackStatsSink;
+		TUniquePtr<FRTPRemoteTrackStatsSink> RemoteAudioTrackStatsSink;
 
-		// All sinks are given all webrtc::RTCStats contained in a report and they can choose which stats they wish to extract/process.
-		TArray<TSharedPtr<FStatsSink>> RTCStatSinks;
+		TUniquePtr<FDataTrackStatsSink>		 DataTrackStatsSink;
+		
+		TUniquePtr<FCandidatePairStatsSink>	 CandidatePairStatsSink;
 	};
-
-#endif // (Migration)
-} // namespace UE::PixelStreaming2
+} // namespace UE::PixelStreaming2
\ No newline at end of file
diff --git a/Engine/Plugins/Media/PixelStreaming2/Source/PixelStreaming2/Private/Stats.cpp b/Engine/Plugins/Media/PixelStreaming2/Source/PixelStreaming2/Private/Stats.cpp
--- a/Engine/Plugins/Media/PixelStreaming2/Source/PixelStreaming2/Private/Stats.cpp
+++ b/Engine/Plugins/Media/PixelStreaming2/Source/PixelStreaming2/Private/Stats.cpp
@@ -47,7 +47,7 @@ namespace UE::PixelStreaming2
 		UConsole::RegisterConsoleAutoCompleteEntries.AddRaw(this, &FStats::UpdateConsoleAutoComplete);
 	}
 
-	void FStats::StorePeerStat(FString PlayerId, FName StatCategory, FStatData Stat)
+	void FStats::StorePeerStat(const FString& PlayerId, FName StatCategory, FStatData Stat)
 	{
 		FName& StatName = Stat.Alias.IsSet() ? Stat.Alias.GetValue() : Stat.StatName;
 
@@ -78,7 +78,7 @@ namespace UE::PixelStreaming2
 		}
 	}
 
-	bool FStats::QueryPeerStat(FString PlayerId, FName StatCategory, FName StatToQuery, double& OutValue) const
+	bool FStats::QueryPeerStat(const FString& PlayerId, FName StatCategory, FName StatToQuery, double& OutValue) const
 	{
 		FScopeLock Lock(&PeerStatsCS);
 
@@ -90,7 +90,7 @@ namespace UE::PixelStreaming2
 		return false;
 	}
 
-	void FStats::RemovePeerStats(FString PlayerId)
+	void FStats::RemovePeerStats(const FString& PlayerId)
 	{
 		FScopeLock Lock(&PeerStatsCS);
 
@@ -197,7 +197,7 @@ namespace UE::PixelStreaming2
 		}
 	}
 
-	void FStats::FireStatChanged(FString PlayerId, FName StatName, float StatValue)
+	void FStats::FireStatChanged(const FString& PlayerId, FName StatName, float StatValue)
 	{
 		// firing off these delegates is not thread safe so we want to mutex this call.
 		FScopeLock Lock(&StatNotificationCS);
diff --git a/Engine/Plugins/Media/PixelStreaming2/Source/PixelStreaming2/Private/Stats.h b/Engine/Plugins/Media/PixelStreaming2/Source/PixelStreaming2/Private/Stats.h
--- a/Engine/Plugins/Media/PixelStreaming2/Source/PixelStreaming2/Private/Stats.h
+++ b/Engine/Plugins/Media/PixelStreaming2/Source/PixelStreaming2/Private/Stats.h
@@ -13,14 +13,6 @@
 
 namespace UE::PixelStreaming2
 {
-	// An interface that allows us to collect webrtc stats from anything that implements it
-	class IStatsSource
-	{
-	public:
-		virtual ~IStatsSource() = default;
-		virtual void PollWebRTCStats() const = 0;
-	};
-
 	struct FStatData
 	{
 	public:
@@ -110,7 +102,7 @@ namespace UE::PixelStreaming2
 	{
 
 	public:
-		FPeerStats(FString InAssociatedPlayer)
+		FPeerStats(const FString& InAssociatedPlayer)
 			: AssociatedPlayer(InAssociatedPlayer)
 			, PlayerIdCanvasItem(FVector2D(0, 0), FText::FromString(FString::Printf(TEXT("[Peer Stats(%s)]"), *AssociatedPlayer)), FSlateFontInfo(FSlateFontInfo(UEngine::GetSmallFont(), 15)), FLinearColor(0, 1, 0))
 		{
@@ -145,10 +137,10 @@ namespace UE::PixelStreaming2
 		static FStats*			Get();
 
 		FStats(const FStats&) = delete;
-		bool QueryPeerStat(FString PlayerId, FName StatCategory, FName StatToQuery, double& OutValue) const;
-		void RemovePeerStats(FString PlayerId);
+		bool QueryPeerStat(const FString& PlayerId, FName StatCategory, FName StatToQuery, double& OutValue) const;
+		void RemovePeerStats(const FString& PlayerId);
 		void RemoveAllPeerStats();
-		void StorePeerStat(FString PlayerId, FName StatCategory, FStatData Stat);
+		void StorePeerStat(const FString& PlayerId, FName StatCategory, FStatData Stat);
 		void StoreApplicationStat(FStatData PeerStat);
 		void Tick(float DeltaTime);
 
@@ -179,8 +171,8 @@ namespace UE::PixelStreaming2
 		FStats();
 		void RegisterEngineHooks();
 		void PollPixelStreaming2Settings();
-		void RemovePeerStat(FString PlayerId);
-		void FireStatChanged(FString PlayerId, FName StatName, float StatValue);
+		void RemovePeerStat(const FString& PlayerId);
+		void FireStatChanged(const FString& PlayerId, FName StatName, float StatValue);
 		void UpdateConsoleAutoComplete(TArray<FAutoCompleteCommand>& AutoCompleteList);
 		void GraphValue_GameThread(FName InName, float Value, int InSamples, float InMinRange, float InMaxRange, float InRefValue);
 		void AddCanvasTile_GameThread(FName Name, const FCanvasTileItem& Tile);
diff --git a/Engine/Plugins/Media/PixelStreaming2/Source/PixelStreaming2/Private/Streamer.cpp b/Engine/Plugins/Media/PixelStreaming2/Source/PixelStreaming2/Private/Streamer.cpp
--- a/Engine/Plugins/Media/PixelStreaming2/Source/PixelStreaming2/Private/Streamer.cpp
+++ b/Engine/Plugins/Media/PixelStreaming2/Source/PixelStreaming2/Private/Streamer.cpp
@@ -69,6 +69,8 @@ namespace UE::PixelStreaming2
 		Streamer->EpicRtcManager->OnDataTrackState.AddSP(Streamer.ToSharedRef(), &FStreamer::OnDataTrackState);
 		Streamer->EpicRtcManager->OnDataTrackMessage.AddSP(Streamer.ToSharedRef(), &FStreamer::OnDataTrackMessage);
 
+		FPixelStreaming2Module::GetModule()->GetStatsCollector()->OnStatsReady.AddSP(Streamer.ToSharedRef(), &FStreamer::OnStatsReady);
+
 		return Streamer;
 	}
 
@@ -260,6 +262,22 @@ namespace UE::PixelStreaming2
 		EpicRtcManager->EpicRtcRoom = nullptr;
 	}
 
+	void FStreamer::OnStatsReady(const FString& PlayerId, const EpicRtcConnectionStats& ConnectionStats)
+	{
+		FPlayerContext* PlayerContext = Players->Find(PlayerId);
+		if (!PlayerContext)
+		{
+			return;
+		}
+
+		if (!PlayerContext->StatsCollector)
+		{
+			return;
+		}
+
+		PlayerContext->StatsCollector->Process(ConnectionStats);
+	}
+
 	IPixelStreaming2Streamer::FPreConnectionEvent& FStreamer::OnPreConnection()
 	{
 		return StreamingPreConnectionEvent;
@@ -722,8 +740,8 @@ namespace UE::PixelStreaming2
 				FStats* Stats = FStats::Get();
 				if (Stats)
 				{
-					Stats->QueryPeerStat(PlayerId, RTCStatTypes::OutboundRTP, PixelStreaming2StatNames::MeanEncodeTime, EncodeMs);
-					Stats->QueryPeerStat(PlayerId, RTCStatTypes::OutboundRTP, PixelStreaming2StatNames::AvgSendDelay, CaptureToSendMs);
+					Stats->QueryPeerStat(PlayerId, RTCStatCategories::LocalVideoTrack, PixelStreaming2StatNames::MeanEncodeTime, EncodeMs);
+					Stats->QueryPeerStat(PlayerId, RTCStatCategories::LocalVideoTrack, PixelStreaming2StatNames::AvgSendDelay, CaptureToSendMs);
 				}
 
 				double TransmissionTimeMs = FPlatformTime::ToMilliseconds64(FPlatformTime::Cycles64());
@@ -927,6 +945,7 @@ namespace UE::PixelStreaming2
 
 		FPlayerContext& PlayerContext = Players->FindOrAdd(ParticipantId);
 		PlayerContext.ParticipantInterface = Participant;
+		PlayerContext.StatsCollector = FRTCStatsCollector::Create(ParticipantId);
 
 		TRefCountPtr<EpicRtcConnectionInterface> ParticipantConnection = Participant->GetConnection();
 		ParticipantConnection->SetManualNegotiation(true);
@@ -965,12 +984,12 @@ namespace UE::PixelStreaming2
 		}
 
 		const EVideoCodec SelectedCodec = GetEnumFromCVar<EVideoCodec>(UPixelStreaming2PluginSettings::CVarEncoderCodec);
-		const bool bNegotiateCodecs = UPixelStreaming2PluginSettings::CVarWebRTCNegotiateCodecs.GetValueOnAnyThread();
-		const bool bTransmitUEVideo = !UPixelStreaming2PluginSettings::CVarWebRTCDisableTransmitVideo.GetValueOnAnyThread();
-		bool bReceiveBrowserVideo = !UPixelStreaming2PluginSettings::CVarWebRTCDisableReceiveVideo.GetValueOnAnyThread();
-		
+		const bool		  bNegotiateCodecs = UPixelStreaming2PluginSettings::CVarWebRTCNegotiateCodecs.GetValueOnAnyThread();
+		const bool		  bTransmitUEVideo = !UPixelStreaming2PluginSettings::CVarWebRTCDisableTransmitVideo.GetValueOnAnyThread();
+		bool			  bReceiveBrowserVideo = !UPixelStreaming2PluginSettings::CVarWebRTCDisableReceiveVideo.GetValueOnAnyThread();
+
 		// Check if the user has selected only H.264 on a AMD gpu and disable receiving video.
-		// WebRTC does not support using SendRecv if the encoding and decoding do not support the same codec. 
+		// WebRTC does not support using SendRecv if the encoding and decoding do not support the same codec.
 		// AMD GPUs currently have decoding disabled so WebRTC fails to create SDP codecs with SendRecv.
 		if (IsRHIDeviceAMD() && !bNegotiateCodecs && SelectedCodec == EVideoCodec::H264)
 		{
diff --git a/Engine/Plugins/Media/PixelStreaming2/Source/PixelStreaming2/Private/Streamer.h b/Engine/Plugins/Media/PixelStreaming2/Source/PixelStreaming2/Private/Streamer.h
--- a/Engine/Plugins/Media/PixelStreaming2/Source/PixelStreaming2/Private/Streamer.h
+++ b/Engine/Plugins/Media/PixelStreaming2/Source/PixelStreaming2/Private/Streamer.h
@@ -13,6 +13,8 @@
 #include "StreamerReconnectTimer.h"
 #include "EpicRtcManager.h"
 
+#include "epic_rtc/core/stats.h"
+
 class IPixelStreaming2Module;
 
 namespace UE::PixelStreaming2
@@ -93,6 +95,7 @@ namespace UE::PixelStreaming2
 		void TriggerMouseLeave(FString InStreamerId);
 		void RemoveSession(bool bDisconnect);
 		void RemoveRoom();
+		void OnStatsReady(const FString& PlayerId, const EpicRtcConnectionStats& ConnectionStats);
 
 	private:
 		FString StreamerId;
diff --git a/Engine/Plugins/Media/PixelStreaming2/Source/PixelStreaming2/Public/PixelStreaming2StatNames.h b/Engine/Plugins/Media/PixelStreaming2/Source/PixelStreaming2/Public/PixelStreaming2StatNames.h
--- a/Engine/Plugins/Media/PixelStreaming2/Source/PixelStreaming2/Public/PixelStreaming2StatNames.h
+++ b/Engine/Plugins/Media/PixelStreaming2/Source/PixelStreaming2/Public/PixelStreaming2StatNames.h
@@ -41,7 +41,6 @@ namespace PixelStreaming2StatNames
 	const FName FirCount				  = FName(TEXT("firCount"));
 	const FName PliCount				  = FName(TEXT("pliCount"));
 	const FName NackCount				  = FName(TEXT("nackCount"));
-	const FName SliCount				  = FName(TEXT("sliCount"));
 	const FName RetransmittedBytesSent	  = FName(TEXT("retransmittedBytesSent"));
 	const FName TargetBitrate			  = FName(TEXT("targetBitrate"));
 	const FName TotalEncodeBytesTarget	  = FName(TEXT("totalEncodedBytesTarget"));
@@ -58,6 +57,17 @@ namespace PixelStreaming2StatNames
 	const FName MessagesSent			  = FName(TEXT("messagesSent"));
 	const FName MessagesReceived		  = FName(TEXT("messagesReceived"));
 
+	const FName PacketsLost				  = FName(TEXT("packetsLost"));
+	const FName Jitter				      = FName(TEXT("jitter"));
+	const FName RoundTripTime			  = FName(TEXT("roundTripTime"));
+	const FName KeyFramesDecoded		  = FName(TEXT("keyFramesDecoded"));
+	const FName AudioLevel				  = FName(TEXT("audioLevel"));
+	const FName TotalSamplesDuration	  = FName(TEXT("totalSamplesDuration"));
+	const FName AvailableOutgoingBitrate  = FName(TEXT("availableOutgoingBitrate"));
+	const FName AvailableIncomingBitrate  = FName(TEXT("availableIncomingBitrate"));
+	const FName RetransmittedBytesReceived	  = FName(TEXT("retransmittedBytesReceived"));
+	const FName RetransmittedPacketsReceived  = FName(TEXT("retransmittedPacketsReceived"));
+
 	// Data-channel stats
 	const FName DataChannelBytesSent		  = FName(TEXT("data-channel-bytesSent"));
 	const FName DataChannelBytesReceived	  = FName(TEXT("data-channel-bytesReceived"));
diff --git a/Engine/Plugins/Media/PixelStreaming2/Source/PixelStreaming2Core/Internal/PixelStreaming2PluginSettings.h b/Engine/Plugins/Media/PixelStreaming2/Source/PixelStreaming2Core/Internal/PixelStreaming2PluginSettings.h
--- a/Engine/Plugins/Media/PixelStreaming2/Source/PixelStreaming2Core/Internal/PixelStreaming2PluginSettings.h
+++ b/Engine/Plugins/Media/PixelStreaming2/Source/PixelStreaming2Core/Internal/PixelStreaming2PluginSettings.h
@@ -368,6 +368,14 @@ public:
 		))
 	bool WebRTCDisableStats = false;
 
+	static TAutoConsoleVariable<float> CVarWebRTCStatsInterval;
+	UPROPERTY(config, EditAnywhere, Category = "WebRtcSettings", meta = (
+		ConsoleVariable = "PixelStreaming2.WebRTC.StatsInterval",
+		DisplayName = "Stats Interval",
+		ToolTip = "Configures how often WebRTC stats are collected."
+		))
+	float WebRTCStatsInterval = 1.f;
+
 	static TAutoConsoleVariable<bool> CVarWebRTCNegotiateCodecs;
 	UPROPERTY(config, EditAnywhere, Category = "WebRtcSettings", meta = (
 		ConsoleVariable = "PixelStreaming2.WebRTC.NegotiateCodecs",
@@ -597,6 +605,9 @@ public:
 		DECLARE_TS_MULTICAST_DELEGATE_OneParam(FOnWebRTCBitrateChanged, IConsoleVariable*);
 		FOnWebRTCBitrateChanged OnWebRTCBitrateChanged;
 
+		DECLARE_TS_MULTICAST_DELEGATE_OneParam(FOnWebRTCDisableStatsChanged, IConsoleVariable*);
+		FOnWebRTCDisableStatsChanged OnWebRTCDisableStatsChanged;
+
 		DECLARE_TS_MULTICAST_DELEGATE_OneParam(FOnLogStatsChanged, IConsoleVariable*);
 		FOnLogStatsChanged OnLogStatsChanged;
 
diff --git a/Engine/Plugins/Media/PixelStreaming2/Source/PixelStreaming2Core/Private/PixelStreaming2PluginSettings.cpp b/Engine/Plugins/Media/PixelStreaming2/Source/PixelStreaming2Core/Private/PixelStreaming2PluginSettings.cpp
--- a/Engine/Plugins/Media/PixelStreaming2/Source/PixelStreaming2Core/Private/PixelStreaming2PluginSettings.cpp
+++ b/Engine/Plugins/Media/PixelStreaming2/Source/PixelStreaming2Core/Private/PixelStreaming2PluginSettings.cpp
@@ -292,6 +292,13 @@ TAutoConsoleVariable<bool> UPixelStreaming2PluginSettings::CVarWebRTCDisableStat
 	TEXT("PixelStreaming2.WebRTC.DisableStats"),
 	false,
 	TEXT("Disables the collection of WebRTC stats."),
+	FConsoleVariableDelegate::CreateLambda([](IConsoleVariable* Var) { Delegates()->OnWebRTCDisableStatsChanged.Broadcast(Var); }),
+	ECVF_Default);
+
+TAutoConsoleVariable<float> UPixelStreaming2PluginSettings::CVarWebRTCStatsInterval(
+	TEXT("PixelStreaming2.WebRTC.StatsInterval"),
+	1.f,
+	TEXT("Configures how often WebRTC stats are collected in seconds. Values less than 0.0f disable stats collection. Default: 1.0f"),
 	ECVF_Default);
 
 TAutoConsoleVariable<bool> UPixelStreaming2PluginSettings::CVarWebRTCNegotiateCodecs(
@@ -825,6 +832,7 @@ static const TMap<FString, FString> GetCmdArg = {
 	{ "WebRTCDisableAudioSync", "PixelStreaming2.WebRTC.DisableAudioSync" },
 	{ "WebRTCEnableFlexFec", "PixelStreaming2.WebRTC.EnableFlexFec" },
 	{ "WebRTCDisableStats", "PixelStreaming2.WebRTC.DisableStats" },
+	{ "WebRTCStatsInterval", "PixelStreaming2.WebRTC.StatsInterval" },
 	{ "WebRTCNegotiateCodecs", "PixelStreaming2.WebRTC.NegotiateCodecs" },
 	{ "WebRTCAudioGain", "PixelStreaming2.WebRTC.AudioGain" },
 	{ "WebRTCPortAllocatorFlags", "PixelStreaming2.WebRTC.PortAllocatorFlags" },
