[PS2] Fix: CVar and launch arg related bugs
- Added legacy command args to support PS1 style launch args
- Removed OnScreenStats cvar. -PixelStreamingOnScreenStats and -PixelStreamingHudStats still work to trigger the stat group on launch
- Fixed PS2 stat groups showing up twice in standalone and packaged

diff --git a/Engine/Plugins/Media/PixelStreaming2/Source/PixelStreaming2/Private/Stats.cpp b/Engine/Plugins/Media/PixelStreaming2/Source/PixelStreaming2/Private/Stats.cpp
--- a/Engine/Plugins/Media/PixelStreaming2/Source/PixelStreaming2/Private/Stats.cpp
+++ b/Engine/Plugins/Media/PixelStreaming2/Source/PixelStreaming2/Private/Stats.cpp
@@ -44,7 +44,8 @@ namespace UE::PixelStreaming2
 	FStats::FStats()
 	{
 		checkf(Instance == nullptr, TEXT("There should only ever been one PixelStreaming2 stats object."));
-		UConsole::RegisterConsoleAutoCompleteEntries.AddRaw(this, &FStats::UpdateConsoleAutoComplete);
+
+		FCoreDelegates::OnPostEngineInit.AddRaw(this, &FStats::RegisterEngineHooks);
 	}
 
 	void FStats::StorePeerStat(const FString& PlayerId, FName StatCategory, FStatData Stat)
@@ -393,33 +394,7 @@ namespace UE::PixelStreaming2
 
 	void FStats::Tick(float DeltaTime)
 	{
-		RTCStatsPolledDelta += DeltaTime;
-
-		// Note (Luke): If we want more metrics from WebRTC there is also the histogram counts.
-		// For example:
-		// RTC_HISTOGRAM_COUNTS("WebRTC.Video.NacksSent", nacks_sent, 1, 100000, 100);
-		// webrtc::metrics::Histogram* Hist1 = webrtc::metrics::HistogramFactoryGetCounts("WebRTC.Video.NacksSent", 0, 100000, 100);
-		// Will require calling webrtc::metrics::Enable();
-
-		// We only poll WebRTC stats every 1s as this matches chrome://webrtc-internals
-		// and more frequency does not seem useful as these stats are mostly used for visual inspection
-		if (RTCStatsPolledDelta > 1.0f)
-		{
-			OnStatsPolled.Broadcast();
-			RTCStatsPolledDelta = 0;
-		}
-
 		PollPixelStreaming2Settings();
-
-		if (!GEngine)
-		{
-			return;
-		}
-
-		if (!bRegisterEngineStats)
-		{
-			RegisterEngineHooks();
-		}
 	}
 
 	void FStats::RemoveAllPeerStats()
@@ -428,62 +403,61 @@ namespace UE::PixelStreaming2
 		PeerStats.Empty();
 	}
 
-	void FStats::ExecStatPS()
-	{
-		// Intetionally empty, registering this function is mostly about getting the stat comment to show up in autocomplete
-	}
-
-	void FStats::ExecStatPSGraphs()
-	{
-		// Intetionally empty, registering this function is mostly about getting the stat comment to show up in autocomplete
-	}
-
 	void FStats::RegisterEngineHooks()
 	{
 		GAreScreenMessagesEnabled = true;
 
 		const FName				   StatName("STAT_PixelStreaming2");
 		const FName				   StatCategory("STATCAT_PixelStreaming2");
-		const FText				   StatDescription(FText::FromString("Pixel Streaming stats for all connected peers."));
+		const FText				   StatDescription(FText::FromString("Stats for the Pixel Streaming plugin and its peers."));
 		UEngine::FEngineStatRender RenderStatFunc = UEngine::FEngineStatRender::CreateRaw(this, &FStats::OnRenderStats);
 		UEngine::FEngineStatToggle ToggleStatFunc = UEngine::FEngineStatToggle::CreateRaw(this, &FStats::OnToggleStats);
 		GEngine->AddEngineStat(StatName, StatCategory, StatDescription, RenderStatFunc, ToggleStatFunc, false);
 
-		// We register this console command so we get autocomplete on `Stat PixelStreaming2`
-		IConsoleManager::Get().RegisterConsoleCommand(
-			TEXT("Stat PixelStreaming2"),
-			TEXT("Stats for the Pixel Streaming plugin and its peers."),
-			FConsoleCommandDelegate::CreateRaw(this, &FStats::ExecStatPS),
-			ECVF_Default);
-
 		const FName				   GraphName("STAT_PixelStreaming2Graphs");
-		const FText				   GraphDescription(FText::FromString("Pixel Streaming graphs showing frame pipeline timings."));
+		const FText				   GraphDescription(FText::FromString("Draws stats graphs for the Pixel Streaming plugin."));
 		UEngine::FEngineStatRender RenderGraphFunc = UEngine::FEngineStatRender::CreateRaw(this, &FStats::OnRenderGraphs);
 		UEngine::FEngineStatToggle ToggleGraphFunc = UEngine::FEngineStatToggle::CreateRaw(this, &FStats::OnToggleGraphs);
 		GEngine->AddEngineStat(GraphName, StatCategory, GraphDescription, RenderGraphFunc, ToggleGraphFunc, false);
 
-		// We register this console command so we get autocomplete on `Stat PixelStreaming2Graphs`
-		IConsoleManager::Get().RegisterConsoleCommand(
-			TEXT("Stat PixelStreaming2Graphs"),
-			TEXT("Draws stats graphs for the Pixel Streaming plugin."),
-			FConsoleCommandDelegate::CreateRaw(this, &FStats::ExecStatPSGraphs),
-			ECVF_Default);
+		UConsole::RegisterConsoleAutoCompleteEntries.AddRaw(this, &FStats::UpdateConsoleAutoComplete);
 
-		bool StatsEnabled = UPixelStreaming2PluginSettings::CVarOnScreenStats.GetValueOnAnyThread();
-		if (StatsEnabled)
-		{
-			for (const FWorldContext& WorldContext : GEngine->GetWorldContexts())
+		// Check the command line for launch args to automatically enable stats for users
+		TFunction<bool(const TCHAR*)> CheckLaunchArgFunc = [](const TCHAR* Match) -> bool {
+			FString ValueMatch(Match);
+			ValueMatch.Append(TEXT("="));
+			FString Value;
+			if (FParse::Value(FCommandLine::Get(), *ValueMatch, Value))
 			{
-				if (WorldContext.WorldType == EWorldType::Game || WorldContext.WorldType == EWorldType::PIE)
+				if (Value.Equals(FString(TEXT("true")), ESearchCase::IgnoreCase))
 				{
-					UWorld*				 World = WorldContext.World();
-					UGameViewportClient* ViewportClient = World->GetGameViewport();
-					GEngine->SetEngineStat(World, ViewportClient, TEXT("PixelStreaming2"), StatsEnabled);
+					return true;
+				}
+				else if (Value.Equals(FString(TEXT("false")), ESearchCase::IgnoreCase))
+				{
+					return false;
 				}
 			}
-		}
+			else if (FParse::Param(FCommandLine::Get(), Match))
+			{
+				return true;
+			}
+
+			return false;
+		};
 
-		bRegisterEngineStats = true;
+		bool bHudStats = CheckLaunchArgFunc(TEXT("PixelStreamingHudStats"));
+		bool bOnScreenStats = CheckLaunchArgFunc(TEXT("PixelStreamingOnScreenStats"));
+
+		if (bHudStats || bOnScreenStats)
+		{
+			for (const FWorldContext& WorldContext : GEngine->GetWorldContexts())
+			{
+				UWorld*				 World = WorldContext.World();
+				UGameViewportClient* ViewportClient = World->GetGameViewport();
+				GEngine->SetEngineStat(World, ViewportClient, TEXT("PixelStreaming2"), true);
+			}
+		}
 	}
 
 	//
diff --git a/Engine/Plugins/Media/PixelStreaming2/Source/PixelStreaming2/Private/Stats.h b/Engine/Plugins/Media/PixelStreaming2/Source/PixelStreaming2/Private/Stats.h
--- a/Engine/Plugins/Media/PixelStreaming2/Source/PixelStreaming2/Private/Stats.h
+++ b/Engine/Plugins/Media/PixelStreaming2/Source/PixelStreaming2/Private/Stats.h
@@ -144,14 +144,9 @@ namespace UE::PixelStreaming2
 		void StoreApplicationStat(FStatData PeerStat);
 		void Tick(float DeltaTime);
 
-		DECLARE_MULTICAST_DELEGATE(FOnStatsPolled);
-		FOnStatsPolled OnStatsPolled;
-
-		void  ExecStatPS();
 		bool  OnToggleStats(UWorld* World, FCommonViewportClient* ViewportClient, const TCHAR* Stream);
 		int32 OnRenderStats(UWorld* World, FViewport* Viewport, FCanvas* Canvas, int32 X, int32 Y, const FVector* ViewLocation, const FRotator* ViewRotation);
 
-		void  ExecStatPSGraphs();
 		bool  OnToggleGraphs(UWorld* World, FCommonViewportClient* ViewportClient, const TCHAR* Stream);
 		int32 OnRenderGraphs(UWorld* World, FViewport* Viewport, FCanvas* Canvas, int32 X, int32 Y, const FVector* ViewLocation, const FRotator* ViewRotation);
 
@@ -169,19 +164,19 @@ namespace UE::PixelStreaming2
 
 	private:
 		FStats();
+
 		void RegisterEngineHooks();
+		void UpdateConsoleAutoComplete(TArray<FAutoCompleteCommand>& AutoCompleteList);
+
 		void PollPixelStreaming2Settings();
 		void RemovePeerStat(const FString& PlayerId);
 		void FireStatChanged(const FString& PlayerId, FName StatName, float StatValue);
-		void UpdateConsoleAutoComplete(TArray<FAutoCompleteCommand>& AutoCompleteList);
 		void GraphValue_GameThread(FName InName, float Value, int InSamples, float InMinRange, float InMaxRange, float InRefValue);
 		void AddCanvasTile_GameThread(FName Name, const FCanvasTileItem& Tile);
 
 	private:
 		static FStats* Instance;
 
-		bool bRegisterEngineStats = false;
-
 		mutable FCriticalSection  PeerStatsCS;
 		TMap<FString, FPeerStats> PeerStats;
 
@@ -189,7 +184,6 @@ namespace UE::PixelStreaming2
 		TMap<FName, FStoredStat> ApplicationStats;
 
 		int64 LastTimeSettingsPolledCycles = 0;
-		float RTCStatsPolledDelta = 0.0f;
 
 		TMap<FName, FDebugGraph>	 Graphs;
 		TMap<FName, FCanvasTileItem> Tiles;
diff --git a/Engine/Plugins/Media/PixelStreaming2/Source/PixelStreaming2Core/Internal/PixelStreaming2PluginSettings.h b/Engine/Plugins/Media/PixelStreaming2/Source/PixelStreaming2Core/Internal/PixelStreaming2PluginSettings.h
--- a/Engine/Plugins/Media/PixelStreaming2/Source/PixelStreaming2Core/Internal/PixelStreaming2PluginSettings.h
+++ b/Engine/Plugins/Media/PixelStreaming2/Source/PixelStreaming2Core/Internal/PixelStreaming2PluginSettings.h
@@ -73,14 +73,6 @@ class PIXELSTREAMING2CORE_API UPixelStreaming2PluginSettings : public UDeveloper
 
 public:
 	// clang-format off
-	static TAutoConsoleVariable<bool> CVarOnScreenStats;
-	UPROPERTY(config, EditAnywhere, Category = "PixelStreaming", meta = (
-		ConsoleVariable = "PixelStreaming2.HUDStats",
-		DisplayName = "Show Pixel Streaming Stats On Screen",
-		ToolTip = "Whether to show PixelStreaming stats on the in-game HUD (default: false)."
-		))
-	bool OnScreenStats = false;
-
 	static TAutoConsoleVariable<bool> CVarLogStats;
 	UPROPERTY(config, EditAnywhere, Category = "PixelStreaming", meta = (
 		ConsoleVariable = "PixelStreaming2.LogStats",
diff --git a/Engine/Plugins/Media/PixelStreaming2/Source/PixelStreaming2Core/Private/PixelStreaming2PluginSettings.cpp b/Engine/Plugins/Media/PixelStreaming2/Source/PixelStreaming2Core/Private/PixelStreaming2PluginSettings.cpp
--- a/Engine/Plugins/Media/PixelStreaming2/Source/PixelStreaming2Core/Private/PixelStreaming2PluginSettings.cpp
+++ b/Engine/Plugins/Media/PixelStreaming2/Source/PixelStreaming2Core/Private/PixelStreaming2PluginSettings.cpp
@@ -70,19 +70,43 @@ namespace
 		// CVars are . deliminated by section. To get their equivilent commandline arg parameter, we need to to remove the .
 		return InCVarName.Replace(TEXT("."), TEXT("")).Replace(TEXT("PixelStreaming2"), TEXT("PixelStreaming"));
 	}
+
+	static void ParseLegacyCommandLineValue(const TCHAR* Match, TAutoConsoleVariable<FString>& CVar)
+	{
+		FString Value;
+		if (FParse::Value(FCommandLine::Get(), Match, Value))
+		{
+			CVar->Set(*Value, ECVF_SetByCommandline);
+		}
+	};
+
+	static void ParseLegacyCommandLineOption(const TCHAR* Match, TAutoConsoleVariable<bool>& CVar)
+	{
+		FString ValueMatch(Match);
+		ValueMatch.Append(TEXT("="));
+		FString Value;
+		if (FParse::Value(FCommandLine::Get(), *ValueMatch, Value))
+		{
+			if (Value.Equals(FString(TEXT("true")), ESearchCase::IgnoreCase))
+			{
+				CVar->Set(true, ECVF_SetByCommandline);
+			}
+			else if (Value.Equals(FString(TEXT("false")), ESearchCase::IgnoreCase))
+			{
+				CVar->Set(false, ECVF_SetByCommandline);
+			}
+		}
+		else if (FParse::Param(FCommandLine::Get(), Match))
+		{
+			CVar->Set(true, ECVF_SetByCommandline);
+		}
+	}
 } // namespace
 
 static FName PixelStreaming2ConsoleVariableMetaFName(TEXT("ConsoleVariable"));
 static FName PixelStreaming2MappedConsoleVariableFName(TEXT("MappedConsoleVariable"));
 
 // Begin Pixel Streaming Plugin CVars
-
-TAutoConsoleVariable<bool> UPixelStreaming2PluginSettings::CVarOnScreenStats(
-	TEXT("PixelStreaming2.HUDStats"),
-	false,
-	TEXT("Whether to show PixelStreaming stats on the in-game HUD (default: false)."),
-	ECVF_Default);
-
 TAutoConsoleVariable<bool> UPixelStreaming2PluginSettings::CVarLogStats(
 	TEXT("PixelStreaming2.LogStats"),
 	false,
@@ -803,7 +827,6 @@ void UPixelStreaming2PluginSettings::SetCVarFromProperty(IConsoleVariable* CVar,
 
 // Map of Property Names to their commandline args as GetMetaData() is not avaliable in packaged projects
 static const TMap<FString, FString> GetCmdArg = {
-	{ "OnScreenStats", "PixelStreaming2.HUDStats" },
 	{ "LogStats", "PixelStreaming2.LogStats" },
 	{ "SendPlayerIdAsInteger", "PixelStreaming2.SendPlayerIdAsInteger" },
 	{ "DisableLatencyTester", "PixelStreaming2.DisableLatencyTester" },
@@ -866,12 +889,16 @@ static const TMap<FString, FString> GetMappedCmdArg = {
 };
 
 static const TArray<FString> GetLegacyCmdArg = {
-	"PixelStreaming2.Encoder.MinQp",
-	"PixelStreaming2.Encoder.MaxQp",
-	"PixelStreaming2.IP",
-	"PixelStreaming2.Port",
-	"PixelStreaming2.URL",
-	"AllowPixelStreamingCommands"
+	"PixelStreaming2.Encoder.MinQp", // Renamed to MaxQuality
+	"PixelStreaming2.Encoder.MaxQp", // Renamed to MinQuality
+	"PixelStreaming2.IP", // Moved to URL
+	"PixelStreaming2.Port", // Moved to URL
+	"PixelStreaming2.URL", // Renamed to SignallingURL
+	"AllowPixelStreamingCommands", 
+	"PixelStreaming2.NegotiateCodecs", // Renamed to WebRTC.NegotiateCodecs
+	"PixelStreaming2.OnScreenStats", // CVar is removed but launch arg is used in stats.cpp
+	"PixelStreaming2.HudStats", // CVar is removed but launch arg is used in stats.cpp
+	"PixelStreaming2.EnableHMD" // Renamed to HMDEnable
 };
 
 void UPixelStreaming2PluginSettings::ValidateCommandLineArgs()
@@ -939,6 +966,7 @@ void UPixelStreaming2PluginSettings::ValidateCommandLineArgs()
 
 void UPixelStreaming2PluginSettings::ParseLegacyCommandlineArgs()
 {
+	// Begin legacy PixelStreaming command line args
 	int32 MinQP;
 	if (FParse::Value(FCommandLine::Get(), TEXT("PixelStreamingEncoderMinQp="), MinQP))
 	{
@@ -978,54 +1006,23 @@ void UPixelStreaming2PluginSettings::ParseLegacyCommandlineArgs()
 		CVarSignallingURL.AsVariable()->Set(*LegacyUrl, ECVF_SetByCommandline);
 	}
 
-	if (FParse::Param(FCommandLine::Get(), TEXT("AllowPixelStreamingCommands")))
-	{
-		CVarInputAllowConsoleCommands->Set(true);
-	}
+	ParseLegacyCommandLineOption(TEXT("PixelStreamingNegotiateCodecs"), CVarWebRTCNegotiateCodecs);
+	ParseLegacyCommandLineOption(TEXT("AllowPixelStreamingCommands"), CVarInputAllowConsoleCommands);
+	ParseLegacyCommandLineOption(TEXT("PixelStreamingDebugDumpFrame"), CVarEncoderDebugDumpFrame);
+	// End legacy PixelStreaming command line args
 
-	// This one gets parsed here because it has no matching UProperty
-	if (FParse::Param(FCommandLine::Get(), TEXT("PixelStreamingEncoderDumpDebugFrames")))
-	{
-		CVarEncoderDebugDumpFrame->Set(true);
-	}
+	// Begin legacy PixelStreamingEditor command line args
+	ParseLegacyCommandLineOption(TEXT("EditorPixelStreamingStartOnLaunch"), CVarEditorStartOnLaunch);
+	ParseLegacyCommandLineOption(TEXT("EditorPixelStreamingUseRemoteSignallingServer"), CVarEditorUseRemoteSignallingServer);
 
-	// Begin legacy EditorStreaming command line args
-	if (FParse::Param(FCommandLine::Get(), TEXT("EditorPixelStreamingStartOnLaunch")))
-	{
-		CVarEditorStartOnLaunch->Set(true);
-	}
+	ParseLegacyCommandLineValue(TEXT("EditorPixelStreamingSource="), CVarEditorSource);
+	IConsoleVariable* EditorSourceCVar = IConsoleManager::Get().FindConsoleVariable(TEXT("PixelStreaming2.Editor.Source"));
+	CheckConsoleEnum<EPixelStreaming2EditorStreamTypes>(EditorSourceCVar);
+	// End legacy PixelStreamingEditor command line args
 
-	FString StartOnLaunchString;
-	if (FParse::Value(FCommandLine::Get(), TEXT("EditorPixelStreamingStartOnLaunch="), StartOnLaunchString))
-	{
-		if (StartOnLaunchString.Equals(FString(TEXT("true")), ESearchCase::IgnoreCase))
-		{
-			CVarEditorStartOnLaunch->Set(true, ECVF_SetByCommandline);
-		}
-		else if (StartOnLaunchString.Equals(FString(TEXT("false")), ESearchCase::IgnoreCase))
-		{
-			CVarEditorStartOnLaunch->Set(false, ECVF_SetByCommandline);
-		}
-	}
-
-	if (FParse::Param(FCommandLine::Get(), TEXT("EditorPixelStreamingUseRemoteSignallingServer")))
-	{
-		CVarEditorUseRemoteSignallingServer->Set(true);
-	}
-
-	FString UseRemoteSignallingServerString;
-	if (FParse::Value(FCommandLine::Get(), TEXT("EditorPixelStreamingUseRemoteSignallingServer="), UseRemoteSignallingServerString))
-	{
-		if (UseRemoteSignallingServerString.Equals(FString(TEXT("true")), ESearchCase::IgnoreCase))
-		{
-			CVarEditorUseRemoteSignallingServer->Set(true, ECVF_SetByCommandline);
-		}
-		else if (UseRemoteSignallingServerString.Equals(FString(TEXT("false")), ESearchCase::IgnoreCase))
-		{
-			CVarEditorUseRemoteSignallingServer->Set(false, ECVF_SetByCommandline);
-		}
-	}
-	// End legacy EditorStreaming command line args
+	// End legacy PixelStreamingHMD command line args
+	ParseLegacyCommandLineOption(TEXT("PixelStreamingEnableHMD"), CVarHMDEnable);
+	// End legacy PixelStreamingHMD command line args
 }
 
 void UPixelStreaming2PluginSettings::PostInitProperties()
diff --git a/Engine/Plugins/Media/PixelStreaming2/Source/PixelStreaming2Editor/Private/PixelStreaming2EditorModule.cpp b/Engine/Plugins/Media/PixelStreaming2/Source/PixelStreaming2Editor/Private/PixelStreaming2EditorModule.cpp
--- a/Engine/Plugins/Media/PixelStreaming2/Source/PixelStreaming2Editor/Private/PixelStreaming2EditorModule.cpp
+++ b/Engine/Plugins/Media/PixelStreaming2/Source/PixelStreaming2Editor/Private/PixelStreaming2EditorModule.cpp
@@ -5,7 +5,7 @@
 #include "PixelStreaming2Style.h"
 #include "PixelStreaming2Toolbar.h"
 #include "AssetTypeActions_VideoProducer.h"
-
+#include "Engine/GameViewportClient.h"
 #include "CoreUtils.h"
 #include "Editor/EditorPerformanceSettings.h"
 #include "Framework/Application/SlateApplication.h"
@@ -79,6 +79,33 @@ namespace UE::EditorPixelStreaming2
 			});
 		}
 
+		// We register console commands for "Stat xxx" here as the autocomplete logic doesn't execute in the editor
+		IConsoleManager::Get().RegisterConsoleCommand(
+			TEXT("Stat PixelStreaming2"),
+			TEXT("Stats for the Pixel Streaming plugin and its peers."),
+			FConsoleCommandDelegate::CreateLambda([]() {
+				for (const FWorldContext& WorldContext : GEngine->GetWorldContexts())
+				{
+					UWorld*				 World = WorldContext.World();
+					UGameViewportClient* ViewportClient = World->GetGameViewport();
+					GEngine->SetEngineStat(World, ViewportClient, TEXT("PixelStreaming2"), true);
+				}
+			}),
+			ECVF_Default);
+
+		IConsoleManager::Get().RegisterConsoleCommand(
+			TEXT("Stat PixelStreaming2Graphs"),
+			TEXT("Draws stats graphs for the Pixel Streaming plugin."),
+			FConsoleCommandDelegate::CreateLambda([]() {
+				for (const FWorldContext& WorldContext : GEngine->GetWorldContexts())
+				{
+					UWorld*				 World = WorldContext.World();
+					UGameViewportClient* ViewportClient = World->GetGameViewport();
+					GEngine->SetEngineStat(World, ViewportClient, TEXT("PixelStreaming2"), true);
+				}
+			}),
+			ECVF_Default);
+
 		IPixelStreaming2Module& Module = IPixelStreaming2Module::Get();
 		Module.OnReady().AddRaw(this, &FPixelStreaming2EditorModule::InitEditorStreaming);
 	}
