[PS2] Add: Ability to specify codec preference order

if negotiating, PS2 will use PixelStreaming2.WebRTC.CodecPreferences for the order
if not negotiating, PS2 will use PixelStreaming2.Encoder.Codec as the only codec

diff --git a/Engine/Plugins/Media/PixelStreaming2/Source/PixelStreaming2/Private/EpicRtcVideoDecoderInitializer.cpp b/Engine/Plugins/Media/PixelStreaming2/Source/PixelStreaming2/Private/EpicRtcVideoDecoderInitializer.cpp
--- a/Engine/Plugins/Media/PixelStreaming2/Source/PixelStreaming2/Private/EpicRtcVideoDecoderInitializer.cpp
+++ b/Engine/Plugins/Media/PixelStreaming2/Source/PixelStreaming2/Private/EpicRtcVideoDecoderInitializer.cpp
@@ -8,6 +8,7 @@
 #include "PixelStreaming2PluginSettings.h"
 #include "ToStringExtensions.h"
 #include "Utils.h"
+#include "VideoUtils.h"
 #include "Video/Decoders/Configs/VideoDecoderConfigAV1.h"
 #include "Video/Decoders/Configs/VideoDecoderConfigH264.h"
 #include "Video/Decoders/Configs/VideoDecoderConfigVP8.h"
@@ -43,9 +44,6 @@ namespace
 
 namespace UE::PixelStreaming2
 {
-
-	TStaticArray<EVideoCodec, 4> FEpicRtcVideoDecoderInitializer::SupportedCodecList = { EVideoCodec::VP8, EVideoCodec::VP9, EVideoCodec::H264, EVideoCodec::AV1 };
-
 	void FEpicRtcVideoDecoderInitializer::CreateDecoder(EpicRtcVideoCodecInfoInterface* CodecInfo, EpicRtcVideoDecoderInterface** OutDecoder)
 	{
 		EpicRtcVideoDecoderInterface* Decoder = nullptr;
@@ -92,11 +90,10 @@ namespace UE::PixelStreaming2
 	// doesnt support the video we are receiving then transport_cc is not enabled which leads to very low bitrate streams.
 	EpicRtcVideoCodecInfoArrayInterface* FEpicRtcVideoDecoderInitializer::GetSupportedCodecs()
 	{
-		// static so we dont create the list every time this is called since the list will not change during runtime.
-		static const TMap<EVideoCodec, TArray<TRefCountPtr<EpicRtcVideoCodecInfoInterface>>> CodecMap = CreateSupportedDecoderMap();
+		const TMap<EVideoCodec, TArray<TRefCountPtr<EpicRtcVideoCodecInfoInterface>>> SupportedCodecMap = CreateSupportedDecoderMap();
 
 		const EVideoCodec SelectedCodec = UE::PixelStreaming2::GetEnumFromCVar<EVideoCodec>(UPixelStreaming2PluginSettings::CVarEncoderCodec);
-		const bool		  NegotiateCodecs = UPixelStreaming2PluginSettings::CVarWebRTCNegotiateCodecs.GetValueOnAnyThread();
+		const bool		  bNegotiateCodecs = UPixelStreaming2PluginSettings::CVarWebRTCNegotiateCodecs.GetValueOnAnyThread();
 
 		// Build a list of supported codecs
 		TArray<TRefCountPtr<EpicRtcVideoCodecInfoInterface>> SupportedCodecs;
@@ -104,37 +101,25 @@ namespace UE::PixelStreaming2
 		// Todo: Check if all HW decoder sessions are used up.
 
 		// If we are not negotiating codecs simply return just the one codec that is selected in UE
-		if (!NegotiateCodecs)
+		if (!bNegotiateCodecs)
 		{
-			if (CodecMap.Contains(SelectedCodec))
+			if (SupportedCodecMap.Contains(SelectedCodec))
 			{
-				SupportedCodecs.Append(CodecMap[SelectedCodec]);
+				SupportedCodecs.Append(SupportedCodecMap[SelectedCodec]);
 			}
 			else
 			{
 				UE_LOG(LogPixelStreaming2, Error, TEXT("Selected codec was not a supported codec."));
 			}
-
-			return new FVideoCodecInfoArray(SupportedCodecs);
 		}
-
-		// Order the codecs so the "selected" codec is first
-		TArray<EVideoCodec> OrderedCodecList;
-		OrderedCodecList.Add(SelectedCodec);
-		for (auto& SupportedCodec : SupportedCodecList)
-		{
-			if (SupportedCodec != SelectedCodec)
-			{
-				OrderedCodecList.Add(SupportedCodec);
-			}
-		}
-
-		// Now attach each codec format to our support codecs list
-		for (auto& Codec : OrderedCodecList)
+		else
 		{
-			if (CodecMap.Contains(Codec))
+			for (auto& Codec : UPixelStreaming2PluginSettings::GetCodecPreferences())
 			{
-				SupportedCodecs.Append(CodecMap[Codec]);
+				if (SupportedCodecMap.Contains(Codec))
+				{
+					SupportedCodecs.Append(SupportedCodecMap[Codec]);
+				}
 			}
 		}
 
@@ -144,7 +129,7 @@ namespace UE::PixelStreaming2
 	TMap<EVideoCodec, TArray<TRefCountPtr<EpicRtcVideoCodecInfoInterface>>> FEpicRtcVideoDecoderInitializer::CreateSupportedDecoderMap()
 	{
 		TMap<EVideoCodec, TArray<TRefCountPtr<EpicRtcVideoCodecInfoInterface>>> Codecs;
-		for (auto& Codec : SupportedCodecList)
+		for (auto& Codec : SupportedVideoCodecs)
 		{
 			Codecs.Add(Codec);
 		}
diff --git a/Engine/Plugins/Media/PixelStreaming2/Source/PixelStreaming2/Private/EpicRtcVideoEncoderInitializer.cpp b/Engine/Plugins/Media/PixelStreaming2/Source/PixelStreaming2/Private/EpicRtcVideoEncoderInitializer.cpp
--- a/Engine/Plugins/Media/PixelStreaming2/Source/PixelStreaming2/Private/EpicRtcVideoEncoderInitializer.cpp
+++ b/Engine/Plugins/Media/PixelStreaming2/Source/PixelStreaming2/Private/EpicRtcVideoEncoderInitializer.cpp
@@ -46,9 +46,6 @@ namespace
 
 namespace UE::PixelStreaming2
 {
-
-	TStaticArray<EVideoCodec, 4> FEpicRtcVideoEncoderInitializer::SupportedCodecList = { EVideoCodec::VP8, EVideoCodec::VP9, EVideoCodec::H264, EVideoCodec::AV1 };
-
 	void FEpicRtcVideoEncoderInitializer::CreateEncoder(EpicRtcVideoCodecInfoInterface* CodecInfo, EpicRtcVideoEncoderInterface** OutEncoder)
 	{
 		EpicRtcVideoEncoderInterface* Encoder = nullptr;
@@ -97,17 +94,18 @@ namespace UE::PixelStreaming2
 	EpicRtcVideoCodecInfoArrayInterface* FEpicRtcVideoEncoderInitializer::GetSupportedCodecs()
 	{
 		// static so we dont create the list every time this is called since the list will not change during runtime.
-		static const TMap<EVideoCodec, TArray<TRefCountPtr<EpicRtcVideoCodecInfoInterface>>> CodecMap = CreateSupportedEncoderMap();
+		static const TMap<EVideoCodec, TArray<TRefCountPtr<EpicRtcVideoCodecInfoInterface>>> SupportedCodecMap = CreateSupportedEncoderMap();
 
-		// Codecs that have been denied (e.g. a codec may be denied if all HW encoder instances are in use)
-		TSet<EVideoCodec> DenyListedCodecs;
+		EVideoCodec SelectedCodec = UE::PixelStreaming2::GetEnumFromCVar<EVideoCodec>(UPixelStreaming2PluginSettings::CVarEncoderCodec);
+		const bool	bNegotiateCodecs = UPixelStreaming2PluginSettings::CVarWebRTCNegotiateCodecs.GetValueOnAnyThread();
 
 		// This array of support codecs gets built up and returned in our preference order
 		TArray<TRefCountPtr<EpicRtcVideoCodecInfoInterface>> SupportedCodecs;
+		// Codecs that have been denied (e.g. a codec may be denied if all HW encoder instances are in use)
+		TSet<EVideoCodec> DenyListedCodecs;
 
-		EVideoCodec SelectedCodec = UE::PixelStreaming2::GetEnumFromCVar<EVideoCodec>(UPixelStreaming2PluginSettings::CVarEncoderCodec);
 #if PLATFORM_WINDOWS || PLATFORM_LINUX
-		if ((SelectedCodec == EVideoCodec::H264 || SelectedCodec == EVideoCodec::AV1) && IsRHIDeviceNVIDIA())
+		if ((SelectedCodec == EVideoCodec::H264 || SelectedCodec == EVideoCodec::AV1 || bNegotiateCodecs) && IsRHIDeviceNVIDIA())
 		{
 			int32 NumEncoderSessions = NvmlEncoder::GetEncoderSessionCount(0); // TODO we should probably actually figure out the GPU index rather than assume 0
 			int32 MaxCVarAllowedSessions = UPixelStreaming2PluginSettings::CVarEncoderMaxSessions.GetValueOnAnyThread();
@@ -146,51 +144,37 @@ namespace UE::PixelStreaming2
 #endif // PLATFORM_WINDOWS || PLATFORM_LINUX
 
 		// If we are not negotiating codecs simply return just the one codec that is selected in UE
-		if (!UPixelStreaming2PluginSettings::CVarWebRTCNegotiateCodecs.GetValueOnAnyThread())
+		if (!bNegotiateCodecs)
 		{
-			SupportedCodecs.Empty();
-
 			if (DenyListedCodecs.Contains(SelectedCodec))
 			{
 				UE_LOG(LogPixelStreaming2, Error, TEXT("Selected codec was denied - most like due to lack of hw encoder sessions."));
 			}
-			else if (CodecMap.Contains(SelectedCodec))
+			else if (SupportedCodecMap.Contains(SelectedCodec))
 			{
-				SupportedCodecs.Append(CodecMap[SelectedCodec]);
+				SupportedCodecs.Append(SupportedCodecMap[SelectedCodec]);
 			}
 			else
 			{
 				UE_LOG(LogPixelStreaming2, Error, TEXT("Selected codec was not a supported codec."));
 			}
-
-			return new FVideoCodecInfoArray(SupportedCodecs);
-		}
-
-		if (UPixelStreaming2PluginSettings::CVarEncoderEnableSimulcast.GetValueOnAnyThread())
-		{
-			// Only H264 and VP8 support simulcast in the way we do it
-			DenyListedCodecs.Add(EVideoCodec::VP9);
-			DenyListedCodecs.Add(EVideoCodec::AV1);
-			UE_LOG(LogPixelStreaming2, Warning, TEXT("Removing VP9 and AV1 from negotiable codecs due to simulcast being enabled"));
 		}
-
-		// order the codecs so the selected is first
-		TArray<EVideoCodec> OrderedCodecList;
-		OrderedCodecList.Add(SelectedCodec);
-		for (auto& SupportedCodec : SupportedCodecList)
+		else
 		{
-			if (SupportedCodec != SelectedCodec && !DenyListedCodecs.Contains(SupportedCodec))
+			if (UPixelStreaming2PluginSettings::CVarEncoderEnableSimulcast.GetValueOnAnyThread())
 			{
-				OrderedCodecList.Add(SupportedCodec);
+				// Only H264 and VP8 support simulcast in the way we do it
+				DenyListedCodecs.Add(EVideoCodec::VP9);
+				DenyListedCodecs.Add(EVideoCodec::AV1);
+				UE_LOG(LogPixelStreaming2, Warning, TEXT("Removing VP9 and AV1 from negotiable codecs due to simulcast being enabled"));
 			}
-		}
 
-		// now just add each of the formats in order
-		for (auto& Codec : OrderedCodecList)
-		{
-			if (CodecMap.Contains(Codec) && !DenyListedCodecs.Contains(Codec))
+			for (auto& Codec : UPixelStreaming2PluginSettings::GetCodecPreferences())
 			{
-				SupportedCodecs.Append(CodecMap[Codec]);
+				if (SupportedCodecMap.Contains(Codec) && !DenyListedCodecs.Contains(Codec))
+				{
+					SupportedCodecs.Append(SupportedCodecMap[Codec]);
+				}
 			}
 		}
 
@@ -200,7 +184,7 @@ namespace UE::PixelStreaming2
 	TMap<EVideoCodec, TArray<TRefCountPtr<EpicRtcVideoCodecInfoInterface>>> FEpicRtcVideoEncoderInitializer::CreateSupportedEncoderMap()
 	{
 		TMap<EVideoCodec, TArray<TRefCountPtr<EpicRtcVideoCodecInfoInterface>>> Codecs;
-		for (auto& Codec : SupportedCodecList)
+		for (auto& Codec : SupportedVideoCodecs)
 		{
 			Codecs.Add(Codec);
 		}
diff --git a/Engine/Plugins/Media/PixelStreaming2/Source/PixelStreaming2/Private/VideoUtils.h b/Engine/Plugins/Media/PixelStreaming2/Source/PixelStreaming2/Private/VideoUtils.h
--- a/Engine/Plugins/Media/PixelStreaming2/Source/PixelStreaming2/Private/VideoUtils.h
+++ b/Engine/Plugins/Media/PixelStreaming2/Source/PixelStreaming2/Private/VideoUtils.h
@@ -26,6 +26,9 @@ namespace UE::PixelStreaming2
 		}
 	};
 
+	// List of video codecs supported by PS2
+	const TStaticArray<EVideoCodec, 4> SupportedVideoCodecs = { EVideoCodec::H264, EVideoCodec::AV1, EVideoCodec::VP8, EVideoCodec::VP9 };
+
 	constexpr uint32_t NumSimulcastLayers = 3;
 	// Each subsequent layer is 1/ScalingFactor the size of the previous
 	constexpr uint32_t ScalingFactor = 2;
diff --git a/Engine/Plugins/Media/PixelStreaming2/Source/PixelStreaming2Core/Internal/PixelStreaming2PluginSettings.h b/Engine/Plugins/Media/PixelStreaming2/Source/PixelStreaming2Core/Internal/PixelStreaming2PluginSettings.h
--- a/Engine/Plugins/Media/PixelStreaming2/Source/PixelStreaming2Core/Internal/PixelStreaming2PluginSettings.h
+++ b/Engine/Plugins/Media/PixelStreaming2/Source/PixelStreaming2Core/Internal/PixelStreaming2PluginSettings.h
@@ -236,7 +236,9 @@ public:
 	UPROPERTY(config, EditAnywhere, Category = "EncoderSettings", meta = (
 		ConsoleVariable = "PixelStreaming2.Encoder.EnableSimulcast",
 		DisplayName = "Enable Simulcast",
-		ToolTip = "Enables simulcast. When enabled, the encoder will encode at full resolution, 1/2 resolution and 1/4 resolution simultaneously. Note: Simulcast is only supported with `H264` and `VP8` and you must use the SFU from the infrastructure to fully utilise this functionality."
+		ToolTip = "Enables simulcast. When enabled, the encoder will encode at full resolution, 1/2 resolution and 1/4 resolution simultaneously. Note: Simulcast is only supported with `H264` and `VP8` and you must use the SFU from the infrastructure to fully utilise this functionality.",
+		EditCondition="!WebRTCNegotiateCodecs",
+		EditConditionHides
 		))
 	bool EnableSimulcast = false;
 
@@ -245,7 +247,9 @@ public:
 		ConsoleVariable = "PixelStreaming2.Encoder.Codec",
 		DisplayName = "Preferred Encoder Codec",
 		ToolTip = "Preferred encoder codec signalled during connection establishment.",
-		GetOptions = "GetVideoCodecOptions"
+		GetOptions = "GetVideoCodecOptions",
+		EditCondition="!WebRTCNegotiateCodecs",
+		EditConditionHides
 		))
 	FString Codec = UE::PixelStreaming2::GetCVarStringFromEnum(EVideoCodec::H264);
 
@@ -254,7 +258,9 @@ public:
 		ConsoleVariable = "PixelStreaming2.Encoder.ScalabilityMode",
 		DisplayName = "Default Scalability Mode",
 		ToolTip = "Indicates number of spatial and temporal layers used, default: L1T1. For a full list of values refer to https://www.w3.org/TR/webrtc-svc/#scalabilitymodes*",
-		GetOptions = "GetScalabilityModeOptions"
+		GetOptions = "GetScalabilityModeOptions",
+		EditCondition="!WebRTCNegotiateCodecs",
+		EditConditionHides
 		))
 	FString ScalabilityMode = UE::PixelStreaming2::GetCVarStringFromEnum(EScalabilityMode::L1T1);
 
@@ -372,10 +378,23 @@ public:
 	UPROPERTY(config, EditAnywhere, Category = "WebRtcSettings", meta = (
 		ConsoleVariable = "PixelStreaming2.WebRTC.NegotiateCodecs",
 		DisplayName = "WebRtc Negotiate Codecs",
-		ToolTip = "Whether PS should send all its codecs during sdp handshake so peers can negotiate or just send a single selected codec."
+		ToolTip = "Whether PixelStreaming should send all its codecs during sdp handshake so peers can negotiate or just send a single selected codec."
 		))
 	bool WebRTCNegotiateCodecs = false;
 
+	static TAutoConsoleVariable<FString> CVarWebRTCCodecPreferences;
+	UPROPERTY(config, EditAnywhere, Category = "WebRtcSettings", meta = (
+		ConsoleVariable = "PixelStreaming2.WebRTC.CodecPreferences",
+		DisplayName = "WebRtc Codec Preferences",
+		ToolTip = "The preference order PixelStreaming will specify during sdp handshake",
+		EditFixedSize,
+		EditCondition="WebRTCNegotiateCodecs",
+		EditConditionHides
+		))
+	TArray<EVideoCodec> WebRTCCodecPreferences = { EVideoCodec::AV1, EVideoCodec::H264, EVideoCodec::VP9, EVideoCodec::VP8 };
+
+	static TArray<EVideoCodec> GetCodecPreferences();
+
 	static TAutoConsoleVariable<float> CVarWebRTCAudioGain;
 	UPROPERTY(config, EditAnywhere, Category = "WebRtcSettings", meta = (
 		ConsoleVariable = "PixelStreaming2.WebRTC.AudioGain",
@@ -443,7 +462,6 @@ public:
 		ToolTip = "Enables the WebRTC-Video-Pacing field trial and sets the video pacing factor parameter. Larger values are more lenient on larger bitrates. Default: -1.0f (values below zero are discarded.)"
 		))
 	float WebRTCVideoPacingFactor = -1.0f;
-	// clang-format on
 
 	// End WebRTC CVars
 
@@ -463,8 +481,8 @@ public:
 		ToolTip = "Enables the use of a remote signalling server. Default: false"
 		))
 	bool EditorUseRemoteSignallingServer = false;
-	
-    static TAutoConsoleVariable<FString> CVarEditorSource;
+
+	static TAutoConsoleVariable<FString> CVarEditorSource;
 	UPROPERTY(config, EditAnywhere, Category = "EditorStreaming", meta = (
 		MappedConsoleVariable = "PixelStreaming2.Editor.Source",
 		DisplayName = "Source",
@@ -474,7 +492,7 @@ public:
 	// End EditorStreaming CVars
 
 	// Begin HMD CVars
-	static TAutoConsoleVariable<bool>  CVarHMDEnable;
+	static TAutoConsoleVariable<bool> CVarHMDEnable;
 	UPROPERTY(config, EditAnywhere, Category = "XR Streaming", meta = (
 		ConsoleVariable = "PixelStreaming2.HMD.Enable",
 		DisplayName = "Enable HMD",
@@ -482,15 +500,15 @@ public:
 		))
 	bool HMDEnable = false;
 
-	static TAutoConsoleVariable<bool>  CVarHMDMatchAspectRatio;
+	static TAutoConsoleVariable<bool> CVarHMDMatchAspectRatio;
 	UPROPERTY(config, EditAnywhere, Category = "XR Streaming", meta = (
 		ConsoleVariable = "PixelStreaming2.HMD.MatchAspectRatio",
 		DisplayName = "Match Aspect Ratio",
 		ToolTip = "If true automatically resize the rendering resolution to match the aspect ratio determined by the HFoV and VFoV. Default: true"
 		))
 	bool HMDMatchAspectRatio = true;
-	
-	static TAutoConsoleVariable<bool>  CVarHMDApplyEyePosition;
+
+	static TAutoConsoleVariable<bool> CVarHMDApplyEyePosition;
 	UPROPERTY(config, EditAnywhere, Category = "XR Streaming", meta = (
 		ConsoleVariable = "PixelStreaming2.HMD.ApplyEyePosition",
 		DisplayName = "Apply Eye Position",
@@ -498,7 +516,7 @@ public:
 		))
 	bool HMDAppleEyePosition = true;
 
-	static TAutoConsoleVariable<bool>  CVarHMDApplyEyeRotation;
+	static TAutoConsoleVariable<bool> CVarHMDApplyEyeRotation;
 	UPROPERTY(config, EditAnywhere, Category = "XR Streaming", meta = (
 		ConsoleVariable = "PixelStreaming2.HMD.ApplyEyeRotation",
 		DisplayName = "Apply Eye Position",
@@ -564,9 +582,10 @@ public:
 		))
 	FString InputKeyFilter = TEXT("");
 
-
 	// End Input CVars
 
+	// clang-format on
+
 	// Begin UDeveloperSettings Interface
 	virtual FName GetCategoryName() const override;
 
diff --git a/Engine/Plugins/Media/PixelStreaming2/Source/PixelStreaming2Core/Private/PixelStreaming2PluginSettings.cpp b/Engine/Plugins/Media/PixelStreaming2/Source/PixelStreaming2Core/Private/PixelStreaming2PluginSettings.cpp
--- a/Engine/Plugins/Media/PixelStreaming2/Source/PixelStreaming2Core/Private/PixelStreaming2PluginSettings.cpp
+++ b/Engine/Plugins/Media/PixelStreaming2/Source/PixelStreaming2Core/Private/PixelStreaming2PluginSettings.cpp
@@ -58,6 +58,39 @@ namespace
 		}
 	}
 
+	static void VerifyCodecPreferenceSettings(UPixelStreaming2PluginSettings* This, FProperty* Property, IConsoleVariable* CVar, const FString& Value = TEXT(""))
+	{
+		TArray<FString> ValidStringArray;
+
+		if (Value != "")
+		{
+			TArray<FString> StringArray;
+			Value.ParseIntoArray(StringArray, TEXT(","));
+
+			for (const FString& CodecString : StringArray)
+			{
+				if (StaticEnum<EVideoCodec>()->GetIndexByNameString(CodecString) == INDEX_NONE)
+				{
+					UE_LOGFMT(LogPixelStreaming2Core, Warning, "Invalid value {0} received for enum of type {1}", CodecString, StaticEnum<EVideoCodec>()->GetName());
+					continue;
+				}
+
+				ValidStringArray.Add(CodecString);
+			}
+		}
+		else
+		{
+			FArrayProperty*		ArrayProperty = CastField<FArrayProperty>(Property);
+			TArray<EVideoCodec> CodecArray = *ArrayProperty->ContainerPtrToValuePtr<TArray<EVideoCodec>>(This);
+			for (EVideoCodec Codec : CodecArray)
+			{
+				ValidStringArray.Add(UE::PixelStreaming2::GetCVarStringFromEnum(Codec));
+			}
+		}
+
+		CVar->Set(*FString::Join(ValidStringArray, TEXT(",")), ECVF_SetByCommandline);
+	}
+
 	FString ConsoleVariableToCommandArgValue(const FString InCVarName)
 	{
 		// CVars are . deliminated by section. To get their equivilent commandline arg for parsing
@@ -331,6 +364,12 @@ TAutoConsoleVariable<bool> UPixelStreaming2PluginSettings::CVarWebRTCNegotiateCo
 	TEXT("Whether PS should send all its codecs during sdp handshake so peers can negotiate or just send a single selected codec."),
 	ECVF_Default);
 
+TAutoConsoleVariable<FString> UPixelStreaming2PluginSettings::CVarWebRTCCodecPreferences(
+	TEXT("PixelStreaming2.WebRTC.CodecPreferences"),
+	TEXT("AV1,H264,VP9,VP8"),
+	TEXT("A comma separated list of video codecs specifying the prefered order PS will signal during sdp handshake"),
+	ECVF_Default);
+
 TAutoConsoleVariable<float> UPixelStreaming2PluginSettings::CVarWebRTCAudioGain(
 	TEXT("PixelStreaming2.WebRTC.AudioGain"),
 	1.0f,
@@ -431,6 +470,27 @@ TAutoConsoleVariable<FString> UPixelStreaming2PluginSettings::CVarInputKeyFilter
 	ECVF_Default);
 // End Input CVars
 
+TArray<EVideoCodec> UPixelStreaming2PluginSettings::GetCodecPreferences()
+{
+	TArray<EVideoCodec> OutCodecPreferences;
+	FString				StringOptions = UPixelStreaming2PluginSettings::CVarWebRTCCodecPreferences.GetValueOnAnyThread();
+	if (StringOptions.IsEmpty())
+	{
+		return OutCodecPreferences;
+	}
+
+	TArray<FString> CodecArray;
+	StringOptions.ParseIntoArray(CodecArray, TEXT(","), true);
+	for (const FString& CodecString : CodecArray)
+	{
+		uint64 EnumIndex = StaticEnum<EVideoCodec>()->GetIndexByNameString(CodecString);
+		checkf(EnumIndex != INDEX_NONE, TEXT("CVar was not containing valid enum string"));
+		OutCodecPreferences.Add(static_cast<EVideoCodec>(StaticEnum<EVideoCodec>()->GetValueByIndex(EnumIndex)));
+	}
+
+	return OutCodecPreferences;
+}
+
 EPortAllocatorFlags UPixelStreaming2PluginSettings::GetPortAllocationFlags()
 {
 	EPortAllocatorFlags OutPortAllocatorFlags = EPortAllocatorFlags::None;
@@ -684,6 +744,11 @@ void UPixelStreaming2PluginSettings::PostEditChangeProperty(FPropertyChangedEven
 				SetPortAllocationCVarFromProperty(this, PropertyChangedEvent.Property);
 			}
 		}
+		else if (PropertyChangedEvent.Property->GetNameCPP() == "WebRTCCodecPreferences")
+		{
+			IConsoleVariable* ConsoleVariable = IConsoleManager::Get().FindConsoleVariable(*PropertyChangedEvent.Property->GetMetaData(PixelStreaming2ConsoleVariableMetaFName));
+			VerifyCodecPreferenceSettings(this, PropertyChangedEvent.Property, ConsoleVariable);
+		}
 		// Codec and ScalabilityMode properties are updated in VerifyVideoSettings once we know all the settings are compatible
 		else if (PropertyChangedEvent.Property->GetNameCPP() != "Codec" && PropertyChangedEvent.Property->GetNameCPP() != "ScalabilityMode")
 		{
@@ -888,17 +953,24 @@ static const TMap<FString, FString> GetMappedCmdArg = {
 	{ "EditorSource", "PixelStreaming2.Editor.Source" }
 };
 
+//                                                                   Property,   CVar,              CommandLine value
+using FMappingFunc = TFunction<void(UPixelStreaming2PluginSettings*, FProperty*, IConsoleVariable*, const FString&)>;
+
+static const TMap<FString, TPair<FString, FMappingFunc>> GetCustomMappedCmdArg = {
+	{ "WebRTCCodecPreferences", { "PixelStreaming2.WebRTC.CodecPreferences", VerifyCodecPreferenceSettings } },
+};
+
 static const TArray<FString> GetLegacyCmdArg = {
 	"PixelStreaming2.Encoder.MinQp", // Renamed to MaxQuality
 	"PixelStreaming2.Encoder.MaxQp", // Renamed to MinQuality
-	"PixelStreaming2.IP", // Moved to URL
-	"PixelStreaming2.Port", // Moved to URL
-	"PixelStreaming2.URL", // Renamed to SignallingURL
-	"AllowPixelStreamingCommands", 
+	"PixelStreaming2.IP",			 // Moved to URL
+	"PixelStreaming2.Port",			 // Moved to URL
+	"PixelStreaming2.URL",			 // Renamed to SignallingURL
+	"AllowPixelStreamingCommands",
 	"PixelStreaming2.NegotiateCodecs", // Renamed to WebRTC.NegotiateCodecs
-	"PixelStreaming2.OnScreenStats", // CVar is removed but launch arg is used in stats.cpp
-	"PixelStreaming2.HudStats", // CVar is removed but launch arg is used in stats.cpp
-	"PixelStreaming2.EnableHMD" // Renamed to HMDEnable
+	"PixelStreaming2.OnScreenStats",   // CVar is removed but launch arg is used in stats.cpp
+	"PixelStreaming2.HudStats",		   // CVar is removed but launch arg is used in stats.cpp
+	"PixelStreaming2.EnableHMD"		   // Renamed to HMDEnable
 };
 
 void UPixelStreaming2PluginSettings::ValidateCommandLineArgs()
@@ -945,6 +1017,18 @@ void UPixelStreaming2PluginSettings::ValidateCommandLineArgs()
 			}
 		}
 
+		if (!bValidArg)
+		{
+			for (const TPair<FString, TPair<FString, FMappingFunc>>& Pair : GetCustomMappedCmdArg)
+			{
+				FString ValidCommandLineArg = ConsoleVariableToCommandArgParam(Pair.Value.Key);
+				if (CurrentCommandLineArg == ValidCommandLineArg)
+				{
+					bValidArg = true;
+				}
+			}
+		}
+
 		if (!bValidArg)
 		{
 			for (const FString& LegacyArg : GetLegacyCmdArg)
@@ -991,7 +1075,7 @@ void UPixelStreaming2PluginSettings::ParseLegacyCommandlineArgs()
 	FString SignallingServerPort;
 	if (FParse::Value(FCommandLine::Get(), TEXT("PixelStreamingPort="), SignallingServerPort))
 	{
-		LegacyUrl += SignallingServerPort;
+		LegacyUrl += TEXT(":") + SignallingServerPort;
 	}
 
 	if (!LegacyUrl.IsEmpty())
@@ -1106,6 +1190,27 @@ void UPixelStreaming2PluginSettings::PostInitProperties()
 				}
 			}
 		}
+
+		// Handle a commandline argument that needs custom mapping from string to some other type
+		// eg TArray to comma separate FString
+		if (GetCustomMappedCmdArg.Contains(Property->GetNameCPP()))
+		{
+			TPair<FString, FMappingFunc> CVarFuncPair = GetCustomMappedCmdArg[Property->GetNameCPP()];
+			FString						 CVarString = CVarFuncPair.Key;
+			FMappingFunc				 MappingFunc = CVarFuncPair.Value;
+
+			FString ConsoleString;
+			if (FParse::Value(FCommandLine::Get(), *ConsoleVariableToCommandArgValue(CVarString), ConsoleString, false))
+			{
+				// Pass in the console string value. This will set the CVar from what was on the command line
+				MappingFunc(this, Property, IConsoleManager::Get().FindConsoleVariable(*CVarString), ConsoleString);
+			}
+			else
+			{
+				// Pass in an empty value. This will set the CVar from the property's value
+				MappingFunc(this, Property, IConsoleManager::Get().FindConsoleVariable(*CVarString), TEXT(""));
+			}
+		}
 	}
 
 	// Handle parsing of legacy command line args (such as -PixelStreamingUrl) after .ini, properties, and new commandline args.
