[PS2] Fix: Audio component playing only through one ear

diff --git a/Engine/Plugins/Media/PixelStreaming2/Source/PixelStreaming2/Private/Blueprints/PixelStreaming2AudioComponent.cpp b/Engine/Plugins/Media/PixelStreaming2/Source/PixelStreaming2/Private/Blueprints/PixelStreaming2AudioComponent.cpp
--- a/Engine/Plugins/Media/PixelStreaming2/Source/PixelStreaming2/Private/Blueprints/PixelStreaming2AudioComponent.cpp
+++ b/Engine/Plugins/Media/PixelStreaming2/Source/PixelStreaming2/Private/Blueprints/PixelStreaming2AudioComponent.cpp
@@ -151,13 +151,36 @@ void UPixelStreaming2AudioComponent::ConsumeRawPCM(const int16_t* AudioData, int
 		GetAudioComponent()->SetPitchMultiplier(1.0f);
 	}
 
-	Audio::TSampleBuffer<int16_t> Buffer(AudioData, NFrames, NChannels, InSampleRate);
-	if (NChannels != SoundGenerator->GetNumChannels())
+	//                                   Data       Samples              Channels   SampleRate
+	Audio::TSampleBuffer<int16_t> Buffer(AudioData, NFrames * NChannels, NChannels, InSampleRate);
+	int32						  TargetNumChannels = SoundGenerator->GetNumChannels();
+	if (NChannels != TargetNumChannels)
 	{
-		Buffer.MixBufferToChannels(SoundGenerator->GetNumChannels());
+		if (NChannels < TargetNumChannels)
+		{
+			// Up mix by duplicating the mono source to each channel
+			TArrayView<int16> SourceBuffer = Buffer.GetArrayView();
+			TArray<int16>	  MixedBuffer;
+			MixedBuffer.SetNumZeroed(NFrames * TargetNumChannels);
+			for (int32 SrcSampleIdx = 0; SrcSampleIdx < Buffer.GetNumSamples(); SrcSampleIdx++)
+			{
+				int32 DestSampleIdx = SrcSampleIdx * TargetNumChannels;
+				for (int32 Channel = 0; Channel < TargetNumChannels; Channel++)
+				{
+					MixedBuffer[DestSampleIdx + Channel] = SourceBuffer[SrcSampleIdx];
+				}
+			}
+
+			Buffer.CopyFrom(MixedBuffer, TargetNumChannels, InSampleRate);
+		}
+		else
+		{
+			// Down mix using inbuilt method
+			Buffer.MixBufferToChannels(TargetNumChannels);
+		}
 	}
 
-	SoundGenerator->AddAudio(Buffer.GetData(), InSampleRate, NChannels, Buffer.GetNumSamples());
+	SoundGenerator->AddAudio(Buffer.GetData(), InSampleRate, Buffer.GetNumChannels(), Buffer.GetNumFrames());
 }
 
 void UPixelStreaming2AudioComponent::OnConsumerAdded()
