[PS2] QOL: Rework dump file logic to enable a new dump file every time the cvar is switched

diff --git a/Engine/Plugins/Media/PixelStreaming2/Source/PixelStreaming2/Private/EpicRtcVideoEncoder.cpp b/Engine/Plugins/Media/PixelStreaming2/Source/PixelStreaming2/Private/EpicRtcVideoEncoder.cpp
--- a/Engine/Plugins/Media/PixelStreaming2/Source/PixelStreaming2/Private/EpicRtcVideoEncoder.cpp
+++ b/Engine/Plugins/Media/PixelStreaming2/Source/PixelStreaming2/Private/EpicRtcVideoEncoder.cpp
@@ -6,7 +6,6 @@
 #include "EpicRtcVideoBufferMultiFormat.h"
 #include "EpicRtcVideoBufferRHI.h"
 #include "EpicRtcVideoCommon.h"
-#include "HAL/PlatformFileManager.h"
 #include "Logging.h"
 #include "PixelCaptureBufferFormat.h"
 #include "PixelCaptureOutputFrameI420.h"
@@ -30,6 +29,16 @@ namespace UE::PixelStreaming2
 	{
 		EncoderConfig._simulcastStreams = nullptr;
 		EncoderConfig._spatialLayers = nullptr;
+
+		if (UPixelStreaming2PluginSettings::CVarEncoderDebugDumpFrame.GetValueOnAnyThread())
+		{
+			CreateDumpFile();
+		}
+
+		if (UPixelStreaming2PluginSettings::FDelegates* Delegates = UPixelStreaming2PluginSettings::Delegates())
+		{
+			DelegateHandle = Delegates->OnEncoderDebugDumpFrameChanged.AddRaw(this, &TEpicRtcVideoEncoder<TVideoResource>::OnEncoderDebugDumpFrameChanged);
+		}
 	}
 
 	template <std::derived_from<FVideoResource> TVideoResource>
@@ -44,6 +53,11 @@ namespace UE::PixelStreaming2
 		{
 			EncoderConfig._spatialLayers->Release();
 		}
+
+		if (UPixelStreaming2PluginSettings::FDelegates* Delegates = UPixelStreaming2PluginSettings::Delegates())
+		{
+			Delegates->OnEncoderDebugDumpFrameChanged.Remove(DelegateHandle);
+		}
 	}
 
 	template <std::derived_from<FVideoResource> TVideoResource>
@@ -670,25 +684,39 @@ namespace UE::PixelStreaming2
 		Encoder->SetMinimalConfig(*VideoConfig);
 	}
 
+	template <std::derived_from<FVideoResource> TVideoResource>
+	void TEpicRtcVideoEncoder<TVideoResource>::OnEncoderDebugDumpFrameChanged(IConsoleVariable* Var)
+	{
+		if (Var->GetBool())
+		{
+			CreateDumpFile();
+		}
+		else
+		{
+			FileHandle = nullptr;
+		}
+	}
+
+	template <std::derived_from<FVideoResource> TVideoResource>
+	void TEpicRtcVideoEncoder<TVideoResource>::CreateDumpFile()
+	{
+		IPlatformFile& PlatformFile = FPlatformFileManager::Get().GetPlatformFile();
+		FString		   TempFilePath = FPaths::CreateTempFilename(*FPaths::ProjectSavedDir(), TEXT("encoded_frame"), TEXT(".raw"));
+		FileHandle = PlatformFile.OpenWrite(*TempFilePath);
+	}
+
 	template <std::derived_from<FVideoResource> TVideoResource>
 	void TEpicRtcVideoEncoder<TVideoResource>::MaybeDumpFrame(EpicRtcEncodedVideoFrame const& EncodedImage)
 	{
 		// Dump encoded frames to file for debugging if CVar is turned on.
-		if (UPixelStreaming2PluginSettings::CVarEncoderDebugDumpFrame.GetValueOnAnyThread())
+		if (!FileHandle)
 		{
-			static IFileHandle* FileHandle = nullptr;
-			if (!FileHandle)
-			{
-				IPlatformFile& PlatformFile = FPlatformFileManager::Get().GetPlatformFile();
-				FString		   TempFilePath = FPaths::CreateTempFilename(*FPaths::ProjectSavedDir(), TEXT("encoded_frame"), TEXT(".raw"));
-				FileHandle = PlatformFile.OpenWrite(*TempFilePath);
-				check(FileHandle);
-				// Note: To examine individual frames from this dump: ffmpeg -i "encoded_frame78134A5047638BB99AE1D88471E5E513.raw" "frames/out-%04d.jpg"
-			}
-
-			FileHandle->Write(EncodedImage._buffer->GetData(), EncodedImage._buffer->GetSize());
-			FileHandle->Flush();
+			return;
 		}
+
+		// Note: To examine individual frames from this dump: ffmpeg -i "encoded_frame78134A5047638BB99AE1D88471E5E513.raw" "frames/out-%04d.jpg"
+		FileHandle->Write(EncodedImage._buffer->GetData(), EncodedImage._buffer->GetSize());
+		FileHandle->Flush();
 	}
 
 	template <std::derived_from<FVideoResource> TVideoResource>
diff --git a/Engine/Plugins/Media/PixelStreaming2/Source/PixelStreaming2/Private/EpicRtcVideoEncoder.h b/Engine/Plugins/Media/PixelStreaming2/Source/PixelStreaming2/Private/EpicRtcVideoEncoder.h
--- a/Engine/Plugins/Media/PixelStreaming2/Source/PixelStreaming2/Private/EpicRtcVideoEncoder.h
+++ b/Engine/Plugins/Media/PixelStreaming2/Source/PixelStreaming2/Private/EpicRtcVideoEncoder.h
@@ -5,6 +5,7 @@
 #include "AVConstants.h"
 #include "EpicRtcMemory.h"
 #include "EpicRtcVideoCommon.h"
+#include "HAL/PlatformFileManager.h"
 #include "Video/Encoders/Configs/VideoEncoderConfigAV1.h"
 #include "Video/Encoders/Configs/VideoEncoderConfigH264.h"
 #include "Video/Encoders/Configs/VideoEncoderConfigVP8.h"
@@ -49,6 +50,9 @@ namespace UE::PixelStreaming2
 		// We store this so we can restore back to it if the user decides to use then stop using the PixelStreaming2.Encoder.TargetBitrate CVar.
 		uint32_t EpicRtcProposedTargetBitrate = 5000000;
 
+		// File handle used to save the encoded output to
+		IFileHandle* FileHandle = nullptr;
+
 	private:
 		bool	 LateInitEncoder();
 		void	 SetInitialSettings(const EpicRtcVideoEncoderConfig& CodecSettings, FVideoEncoderConfig& VideoConfig);
@@ -56,6 +60,10 @@ namespace UE::PixelStreaming2
 		void	 MaybeDumpFrame(EpicRtcEncodedVideoFrame const& EncodedImage);
 		uint32_t SumAndResetBitrates();
 
+		void OnEncoderDebugDumpFrameChanged(IConsoleVariable* Var);
+		void CreateDumpFile();
+		FDelegateHandle DelegateHandle;
+
 	public:
 		/* Begin EpicRtcRefCountInterface */
 		virtual uint32_t AddRef() override final { return TRefCountingMixin<TEpicRtcVideoEncoder>::AddRef(); }
diff --git a/Engine/Plugins/Media/PixelStreaming2/Source/PixelStreaming2Core/Internal/PixelStreaming2PluginSettings.h b/Engine/Plugins/Media/PixelStreaming2/Source/PixelStreaming2Core/Internal/PixelStreaming2PluginSettings.h
--- a/Engine/Plugins/Media/PixelStreaming2/Source/PixelStreaming2Core/Internal/PixelStreaming2PluginSettings.h
+++ b/Engine/Plugins/Media/PixelStreaming2/Source/PixelStreaming2Core/Internal/PixelStreaming2PluginSettings.h
@@ -624,6 +624,9 @@ public:
 
 		DECLARE_TS_MULTICAST_DELEGATE_OneParam(FOnInputKeyFilterChanged, IConsoleVariable*);
 		FOnInputKeyFilterChanged OnInputKeyFilterChanged;
+
+		DECLARE_TS_MULTICAST_DELEGATE_OneParam(FOnEncoderDebugDumpFrameChanged, IConsoleVariable*);
+		FOnEncoderDebugDumpFrameChanged OnEncoderDebugDumpFrameChanged;
 	};
 
 	static FDelegates* Delegates();
diff --git a/Engine/Plugins/Media/PixelStreaming2/Source/PixelStreaming2Core/Private/PixelStreaming2PluginSettings.cpp b/Engine/Plugins/Media/PixelStreaming2/Source/PixelStreaming2Core/Private/PixelStreaming2PluginSettings.cpp
--- a/Engine/Plugins/Media/PixelStreaming2/Source/PixelStreaming2Core/Private/PixelStreaming2PluginSettings.cpp
+++ b/Engine/Plugins/Media/PixelStreaming2/Source/PixelStreaming2Core/Private/PixelStreaming2PluginSettings.cpp
@@ -277,6 +277,7 @@ TAutoConsoleVariable<bool> UPixelStreaming2PluginSettings::CVarEncoderDebugDumpF
 	TEXT("PixelStreaming2.Encoder.DumpDebugFrames"),
 	false,
 	TEXT("Dumps frames from the encoder to a file on disk for debugging purposes."),
+	FConsoleVariableDelegate::CreateLambda([](IConsoleVariable* Var) { Delegates()->OnEncoderDebugDumpFrameChanged.Broadcast(Var); }),
 	ECVF_Default);
 
 // Begin WebRTC CVars
