[PS2] Add: Max FPS support when decoupled

diff --git a/Engine/Plugins/Media/PixelStreaming2/Source/PixelStreaming2/Private/Blueprints/PixelStreaming2StreamerComponent.cpp b/Engine/Plugins/Media/PixelStreaming2/Source/PixelStreaming2/Private/Blueprints/PixelStreaming2StreamerComponent.cpp
--- a/Engine/Plugins/Media/PixelStreaming2/Source/PixelStreaming2/Private/Blueprints/PixelStreaming2StreamerComponent.cpp
+++ b/Engine/Plugins/Media/PixelStreaming2/Source/PixelStreaming2/Private/Blueprints/PixelStreaming2StreamerComponent.cpp
@@ -10,7 +10,7 @@ UPixelStreaming2StreamerComponent::UPixelStreaming2StreamerComponent(const FObje
 {
 }
 
-void UPixelStreaming2StreamerComponent::BeginPlay() 
+void UPixelStreaming2StreamerComponent::BeginPlay()
 {
 	if (Streamer)
 	{
@@ -52,7 +52,12 @@ void UPixelStreaming2StreamerComponent::StartStreaming()
 	}
 
 	Streamer->SetCoupleFramerate(CoupleFramerate);
-	Streamer->SetStreamFPS(StreamFPS);
+
+	if (StreamFPSOverride > 0)
+	{
+		Streamer->SetStreamFPS(StreamFPSOverride);
+	}
+
 	if (UsePixelStreamingURL)
 	{
 		FString ServerURL;
@@ -145,7 +150,7 @@ void UPixelStreaming2StreamerComponent::SetupStreamerInput()
 		// default to the scene viewport if we have a game engine
 		if (UGameEngine* GameEngine = Cast<UGameEngine>(GEngine))
 		{
-			TSharedPtr<FSceneViewport>					   TargetViewport = GameEngine->SceneViewport;
+			TSharedPtr<FSceneViewport>				 TargetViewport = GameEngine->SceneViewport;
 			TSharedPtr<IPixelStreaming2InputHandler> InputHandler = Streamer->GetInputHandler().Pin();
 			if (TargetViewport.IsValid() && InputHandler.IsValid())
 			{
diff --git a/Engine/Plugins/Media/PixelStreaming2/Source/PixelStreaming2/Private/Blueprints/PixelStreaming2StreamerComponent.h b/Engine/Plugins/Media/PixelStreaming2/Source/PixelStreaming2/Private/Blueprints/PixelStreaming2StreamerComponent.h
--- a/Engine/Plugins/Media/PixelStreaming2/Source/PixelStreaming2/Private/Blueprints/PixelStreaming2StreamerComponent.h
+++ b/Engine/Plugins/Media/PixelStreaming2/Source/PixelStreaming2/Private/Blueprints/PixelStreaming2StreamerComponent.h
@@ -63,7 +63,7 @@ public:
 	bool UsePixelStreamingURL = false;
 
 	UPROPERTY(EditAnywhere, Category = "PixelStreaming2")
-	int32 StreamFPS = 60;
+	int32 StreamFPSOverride = -1;
 
 	UPROPERTY(EditAnywhere, Category = "PixelStreaming2")
 	bool CoupleFramerate = false;
diff --git a/Engine/Plugins/Media/PixelStreaming2/Source/PixelStreaming2/Private/Streamer.cpp b/Engine/Plugins/Media/PixelStreaming2/Source/PixelStreaming2/Private/Streamer.cpp
--- a/Engine/Plugins/Media/PixelStreaming2/Source/PixelStreaming2/Private/Streamer.cpp
+++ b/Engine/Plugins/Media/PixelStreaming2/Source/PixelStreaming2/Private/Streamer.cpp
@@ -125,7 +125,7 @@ namespace UE::PixelStreaming2
 
 	void FStreamer::SetCoupleFramerate(bool bCouple)
 	{
-		VideoSourceGroup->SetCoupleFramerate(bCouple);
+		VideoSourceGroup->SetDecoupleFramerate(!bCouple);
 	}
 
 	void FStreamer::SetVideoProducer(TSharedPtr<IPixelStreaming2VideoProducer> Producer)
diff --git a/Engine/Plugins/Media/PixelStreaming2/Source/PixelStreaming2/Private/VideoSourceGroup.cpp b/Engine/Plugins/Media/PixelStreaming2/Source/PixelStreaming2/Private/VideoSourceGroup.cpp
--- a/Engine/Plugins/Media/PixelStreaming2/Source/PixelStreaming2/Private/VideoSourceGroup.cpp
+++ b/Engine/Plugins/Media/PixelStreaming2/Source/PixelStreaming2/Private/VideoSourceGroup.cpp
@@ -10,17 +10,150 @@
 
 namespace UE::PixelStreaming2
 {
+	class FFrameRunnable : public FRunnable, public FSingleThreadRunnable
+	{
+	public:
+		FFrameRunnable(TWeakPtr<FVideoSourceGroup> VideoSourceGroup, bool bEnabled)
+			: bEnabled(bEnabled)
+			, OuterVideoSourceGroup(VideoSourceGroup)
+		{
+		}
+
+		virtual ~FFrameRunnable() = default;
+
+		void SetEnabled(bool bInEnabled)
+		{
+			bEnabled = bInEnabled;
+
+			if (bEnabled)
+			{
+				// The thread has enabled so wake it up
+				FrameEvent->Trigger();
+			}
+		}
+
+		virtual bool Init() override
+		{
+			return true;
+		}
+
+		virtual uint32 Run() override
+		{
+			bIsRunning = true;
+
+			while (bIsRunning)
+			{
+				if (!bEnabled)
+				{
+					// Sleep the thread indefinitely because it is not enabled
+					FrameEvent->Wait();
+				}
+
+				if (TSharedPtr<FVideoSourceGroup> VideoSourceGroup = OuterVideoSourceGroup.Pin())
+				{
+
+					const double TimeSinceLastSubmitMs = FPlatformTime::ToMilliseconds64(FPlatformTime::Cycles64() - LastSubmitCycles);
+
+					// Decrease this value to make expected frame delivery more precise, however may result in more old frames being sent
+					const double PrecisionFactor = 0.1;
+					const double WaitFactor = UPixelStreaming2PluginSettings::CVarDecoupleWaitFactor.GetValueOnAnyThread();
+
+					// In "auto" mode vary this value based on historical average
+					const double TargetSubmitMs = 1000.0 / VideoSourceGroup->GetFPS();
+					const double TargetSubmitMsWithPadding = TargetSubmitMs * WaitFactor;
+					const double MinSubmitMs = TargetSubmitMs * (1.f / WaitFactor);
+					const double CloseEnoughMs = TargetSubmitMs * PrecisionFactor;
+					const bool	 bFrameOverdue = TimeSinceLastSubmitMs >= TargetSubmitMsWithPadding;
+
+					// Check frame arrived in time
+					if (!bFrameOverdue)
+					{
+						// Frame arrived in a timely fashion, but is it too soon to maintain our target rate? If so, sleep.
+						double WaitTimeRemainingMs = TargetSubmitMsWithPadding - TimeSinceLastSubmitMs;
+						if (WaitTimeRemainingMs > CloseEnoughMs)
+						{
+							bool bGotNewFrame = FrameEvent->Wait(WaitTimeRemainingMs);
+							if (!bGotNewFrame)
+							{
+								UE_LOG(LogPixelStreaming2, VeryVerbose, TEXT("Old frame submitted"));
+							}
+						}
+					}
+
+					// Push frame immediately
+					PushFrame(VideoSourceGroup);
+				}
+			}
+
+			return 0;
+		}
+
+		virtual void Stop() override
+		{
+			bIsRunning = false;
+		}
+
+		virtual void Exit() override
+		{
+			bIsRunning = false;
+		}
+
+		virtual FSingleThreadRunnable* GetSingleThreadInterface() override
+		{
+			bIsRunning = true;
+			return this;
+		}
+
+		/*
+		 * Note this function is required as part of `FSingleThreadRunnable` and only gets called when engine is run in single-threaded mode,
+		 * so the logic is much less complex as this is not a case we particularly optimize for, a simple tick on an interval will be acceptable.
+		 */
+		virtual void Tick() override
+		{
+			if (TSharedPtr<FVideoSourceGroup> VideoSourceGroup = OuterVideoSourceGroup.Pin())
+			{
+				const uint64 NowCycles = FPlatformTime::Cycles64();
+				const double DeltaMs = FPlatformTime::ToMilliseconds64(NowCycles - LastSubmitCycles);
+				const double TargetSubmitMs = 1000.0 / VideoSourceGroup->GetFPS();
+				if (DeltaMs >= TargetSubmitMs)
+				{
+					PushFrame(VideoSourceGroup);
+				}
+			}
+		}
+
+		void PushFrame(TSharedPtr<FVideoSourceGroup> VideoSourceGroup)
+		{
+			VideoSourceGroup->PushFrame();
+			LastSubmitCycles = FPlatformTime::Cycles64();
+		}
+
+		bool						bEnabled = false;
+		bool						bIsRunning = false;
+		TWeakPtr<FVideoSourceGroup> OuterVideoSourceGroup = nullptr;
+		uint64						LastSubmitCycles = 0;
+
+		/* Use this event to signal when we should wake and also how long we should sleep for between transmitting a frame. */
+		FEventRef FrameEvent;
+	};
+
 	TSharedPtr<FVideoSourceGroup> FVideoSourceGroup::Create(TSharedPtr<FVideoCapturer> InVideoCapturer)
 	{
 		TSharedPtr<FVideoSourceGroup> VideoSourceGroup = TSharedPtr<FVideoSourceGroup>(new FVideoSourceGroup());
 
 		VideoSourceGroup->FrameDelegateHandle = InVideoCapturer->OnFrameCaptured.AddSP(VideoSourceGroup.ToSharedRef(), &FVideoSourceGroup::OnFrameCaptured);
 
+		if (UPixelStreaming2PluginSettings::FDelegates* Delegates = UPixelStreaming2PluginSettings::Delegates())
+		{
+			VideoSourceGroup->FpsDelegateHandle = Delegates->OnWebRTCFpsChanged.AddSP(VideoSourceGroup.ToSharedRef(), &FVideoSourceGroup::OnWebRtcFpsChanged);
+			VideoSourceGroup->DecoupleDelegateHandle = Delegates->OnDecoupleFramerateChanged.AddSP(VideoSourceGroup.ToSharedRef(), &FVideoSourceGroup::OnDecoupleFramerateChanged);
+		}
+
 		return VideoSourceGroup;
 	}
 
 	FVideoSourceGroup::FVideoSourceGroup()
-		: bCoupleFramerate(!UPixelStreaming2PluginSettings::CVarDecoupleFramerate.GetValueOnAnyThread())
+		: bDecoupleFramerate(UPixelStreaming2PluginSettings::CVarDecoupleFramerate.GetValueOnAnyThread())
 		, FramesPerSecond(UPixelStreaming2PluginSettings::CVarWebRTCFps.GetValueOnAnyThread())
 	{
 	}
@@ -32,6 +165,8 @@ namespace UE::PixelStreaming2
 
 	void FVideoSourceGroup::SetFPS(int32 InFramesPerSecond)
 	{
+		bFPSOverride = true;
+
 		FramesPerSecond = InFramesPerSecond;
 	}
 
@@ -40,9 +175,43 @@ namespace UE::PixelStreaming2
 		return FramesPerSecond;
 	}
 
-	void FVideoSourceGroup::SetCoupleFramerate(bool Couple)
+	void FVideoSourceGroup::OnWebRtcFpsChanged(IConsoleVariable* Var)
+	{
+		// User has manually overriden the stream fps so don't respect the cvar change
+		if (bFPSOverride)
+		{
+			return;
+		}
+
+		FramesPerSecond = Var->GetInt();
+	}
+
+	void FVideoSourceGroup::SetDecoupleFramerate(bool bDecouple)
+	{
+		bDecoupleOverride = true;
+
+		bDecoupleFramerate = bDecouple;
+
+		if (FrameRunnable)
+		{
+			FrameRunnable->SetEnabled(bDecoupleFramerate);
+		}
+	}
+
+	void FVideoSourceGroup::OnDecoupleFramerateChanged(IConsoleVariable* Var)
 	{
-		bCoupleFramerate = Couple;
+		// User has manually overriden the decouple bool so don't respect the cvar change
+		if (bDecoupleOverride)
+		{
+			return;
+		}
+
+		bDecoupleFramerate = Var->GetBool();
+
+		if (FrameRunnable)
+		{
+			FrameRunnable->SetEnabled(bDecoupleFramerate);
+		}
 	}
 
 	void FVideoSourceGroup::AddVideoSource(TSharedPtr<FEpicRtcVideoSource> VideoSource)
@@ -51,8 +220,6 @@ namespace UE::PixelStreaming2
 			FScopeLock Lock(&CriticalSection);
 			VideoSources.Add(VideoSource);
 		}
-
-		CheckStartStopThread();
 	}
 
 	void FVideoSourceGroup::RemoveVideoSource(const FEpicRtcVideoSource* ToRemove)
@@ -63,7 +230,6 @@ namespace UE::PixelStreaming2
 				return Target.Get() == ToRemove;
 			});
 		}
-		CheckStartStopThread();
 	}
 
 	void FVideoSourceGroup::RemoveAllVideoSources()
@@ -72,32 +238,28 @@ namespace UE::PixelStreaming2
 			FScopeLock Lock(&CriticalSection);
 			VideoSources.Empty();
 		}
-		CheckStartStopThread();
 	}
 
 	void FVideoSourceGroup::Start()
 	{
-		if (!bRunning)
-		{
-			if (VideoSources.Num() > 0)
-			{
-				StartThread();
-			}
-			bRunning = true;
-		}
+		bRunning = true;
+		StartThread();
 	}
 
 	void FVideoSourceGroup::Stop()
 	{
-		if (bRunning)
-		{
-			StopThread();
-			bRunning = false;
-		}
+		bRunning = false;
+		StopThread();
 	}
 
-	void FVideoSourceGroup::Tick()
+	void FVideoSourceGroup::PushFrame()
 	{
+		// The FrameRunnable may push a frame post engine exit because it runs in its own thread
+		if (IsEngineExitRequested())
+		{
+			return;
+		}
+
 		TRACE_CPUPROFILER_EVENT_SCOPE_ON_CHANNEL_STR("PixelStreaming2 Video Source Group Tick", PixelStreaming2Channel);
 		FScopeLock Lock(&CriticalSection);
 
@@ -113,128 +275,40 @@ namespace UE::PixelStreaming2
 
 	void FVideoSourceGroup::OnFrameCaptured()
 	{
-		if (bCoupleFramerate)
+		if (!bDecoupleFramerate)
 		{
-			Tick();
-		}
-		else
-		{
-			// We are in decoupled render/streaming mode - we should wake our VideoSourceGroup::FFrameThread (if it is sleeping)
-			if (FrameRunnable && FrameRunnable->FrameEvent.Get())
-			{
-				FrameRunnable->FrameEvent.Get()->Trigger();
-			}
+			// Source FPS and engine FPS are coupled. Manually push a frame
+			PushFrame();
 		}
 	}
 
 	void FVideoSourceGroup::StartThread()
 	{
-		if (!bCoupleFramerate && !bThreadRunning)
+		if (!FrameRunnable)
 		{
-			FrameRunnable = MakeUnique<FFrameThread>(AsWeak());
-			FrameThread = FRunnableThread::Create(FrameRunnable.Get(), TEXT("FVideoSourceGroup Thread"), 0, TPri_TimeCritical);
-			bThreadRunning = true;
+			FrameRunnable = MakeShared<FFrameRunnable>(AsWeak(), bDecoupleFramerate);
 		}
-	}
 
-	void FVideoSourceGroup::StopThread()
-	{
-		if (FrameThread != nullptr)
+		if (!FrameThread)
 		{
-			FrameThread->Kill(true);
-		}
-		FrameThread = nullptr;
-		bThreadRunning = false;
-	}
-
-	void FVideoSourceGroup::CheckStartStopThread()
-	{
-		if (bRunning)
-		{
-			const int32 NumSources = VideoSources.Num();
-			if (bThreadRunning && NumSources == 0)
-			{
-				StopThread();
-			}
-			else if (!bThreadRunning && NumSources > 0)
-			{
-				StartThread();
-			}
+			FrameThread = TSharedPtr<FRunnableThread>(FRunnableThread::Create(FrameRunnable.Get(), TEXT("FVideoSourceGroup Thread"), 0, TPri_TimeCritical));
 		}
 	}
 
-	bool FVideoSourceGroup::FFrameThread::Init()
-	{
-		return true;
-	}
-
-	uint32 FVideoSourceGroup::FFrameThread::Run()
+	void FVideoSourceGroup::StopThread()
 	{
-		bIsRunning = true;
-
-		while (bIsRunning)
+		if (FrameRunnable)
 		{
-			if (TSharedPtr<FVideoSourceGroup> VideoSourceGroup = OuterVideoSourceGroup.Pin())
-			{
-				const double TimeSinceLastSubmitMs = FPlatformTime::ToMilliseconds64(FPlatformTime::Cycles64() - LastSubmitCycles);
-
-				// Decrease this value to make expected frame delivery more precise, however may result in more old frames being sent
-				const double PrecisionFactor = 0.1;
-				const double WaitFactor = UPixelStreaming2PluginSettings::CVarDecoupleWaitFactor.GetValueOnAnyThread();
-
-				// In "auto" mode vary this value based on historical average
-				const double TargetSubmitMs = 1000.0 / VideoSourceGroup->FramesPerSecond;
-				const double TargetSubmitMsWithPadding = TargetSubmitMs * WaitFactor;
-				const double CloseEnoughMs = TargetSubmitMs * PrecisionFactor;
-				const bool	 bFrameOverdue = TimeSinceLastSubmitMs >= TargetSubmitMsWithPadding;
-
-				// Check frame arrived in time
-				if (!bFrameOverdue)
-				{
-					// Frame arrived in a timely fashion, but is it too soon to maintain our target rate? If so, sleep.
-					double WaitTimeRemainingMs = TargetSubmitMsWithPadding - TimeSinceLastSubmitMs;
-					if (WaitTimeRemainingMs > CloseEnoughMs)
-					{
-						bool bGotNewFrame = FrameEvent.Get()->Wait(WaitTimeRemainingMs);
-						if (!bGotNewFrame)
-						{
-							UE_LOG(LogPixelStreaming2, VeryVerbose, TEXT("Old frame submitted"));
-						}
-					}
-				}
-
-				// Push frame immediately
-				PushFrame(VideoSourceGroup);
-			}
+			FrameRunnable->Stop();
+			// Wake the thread in case it was sleeping. This will cause it to then exit the ::Run loop
+			FrameRunnable->FrameEvent->Trigger();
+			FrameRunnable.Reset();
 		}
-		return 0;
-	}
-
-	void FVideoSourceGroup::FFrameThread::Stop()
-	{
-		bIsRunning = false;
-	}
-
-	void FVideoSourceGroup::FFrameThread::Exit()
-	{
-		bIsRunning = false;
-	}
-
-	/*
-	 * Note this function is required as part of `FSingleThreadRunnable` and only gets called when engine is run in single-threaded mode,
-	 * so the logic is much less complex as this is not a case we particularly optimize for, a simple tick on an interval will be acceptable.
-	 */
-	void FVideoSourceGroup::FFrameThread::Tick()
-	{
-		if (TSharedPtr<FVideoSourceGroup> VideoSourceGroup = OuterVideoSourceGroup.Pin())
+		
+		if (FrameThread)
 		{
-			const uint64 NowCycles = FPlatformTime::Cycles64();
-			const double DeltaMs = FPlatformTime::ToMilliseconds64(NowCycles - LastSubmitCycles);
-			const double TargetSubmitMs = 1000.0 / VideoSourceGroup->FramesPerSecond;
-			if (DeltaMs >= TargetSubmitMs)
-			{
-				PushFrame(VideoSourceGroup);
-			}
+			FrameThread->Kill(true);
+			FrameThread.Reset();
 		}
 	}
 
@@ -256,10 +330,4 @@ namespace UE::PixelStreaming2
 			}
 		}
 	}
-
-	void FVideoSourceGroup::FFrameThread::PushFrame(TSharedPtr<FVideoSourceGroup> VideoSourceGroup)
-	{
-		VideoSourceGroup->Tick();
-		LastSubmitCycles = FPlatformTime::Cycles64();
-	}
 } // namespace UE::PixelStreaming2
diff --git a/Engine/Plugins/Media/PixelStreaming2/Source/PixelStreaming2/Private/VideoSourceGroup.h b/Engine/Plugins/Media/PixelStreaming2/Source/PixelStreaming2/Private/VideoSourceGroup.h
--- a/Engine/Plugins/Media/PixelStreaming2/Source/PixelStreaming2/Private/VideoSourceGroup.h
+++ b/Engine/Plugins/Media/PixelStreaming2/Source/PixelStreaming2/Private/VideoSourceGroup.h
@@ -3,6 +3,7 @@
 #pragma once
 
 #include "HAL/Event.h"
+#include "HAL/IConsoleManager.h"
 #include "HAL/Runnable.h"
 #include "HAL/RunnableThread.h"
 #include "Misc/SingleThreadRunnable.h"
@@ -12,6 +13,7 @@
 namespace UE::PixelStreaming2
 {
 	class FEpicRtcVideoSource;
+	class FFrameRunnable;
 
 	class FVideoSourceGroup : public TSharedFromThis<FVideoSourceGroup>
 	{
@@ -22,7 +24,7 @@ namespace UE::PixelStreaming2
 		void  SetFPS(int32 InFramesPerSecond);
 		int32 GetFPS();
 
-		void SetCoupleFramerate(bool Couple);
+		void SetDecoupleFramerate(bool bDecouple);
 
 		void AddVideoSource(TSharedPtr<FEpicRtcVideoSource> VideoSource);
 		void RemoveVideoSource(const FEpicRtcVideoSource* ToRemove);
@@ -30,8 +32,7 @@ namespace UE::PixelStreaming2
 
 		void Start();
 		void Stop();
-		void Tick();
-		bool IsThreadRunning() const { return bRunning; }
+		void PushFrame();
 
 		void ForceKeyFrame();
 
@@ -40,50 +41,27 @@ namespace UE::PixelStreaming2
 
 		void StartThread();
 		void StopThread();
-		void CheckStartStopThread();
+
 		void OnFrameCaptured();
 
-		class FFrameThread : public FRunnable, public FSingleThreadRunnable
-		{
-		public:
-			FFrameThread(TWeakPtr<FVideoSourceGroup> InVideoSourceGroup)
-				: OuterVideoSourceGroup(InVideoSourceGroup)
-			{
-			}
-			virtual ~FFrameThread() = default;
-
-			virtual bool   Init() override;
-			virtual uint32 Run() override;
-			virtual void   Stop() override;
-			virtual void   Exit() override;
-
-			virtual FSingleThreadRunnable* GetSingleThreadInterface() override
-			{
-				bIsRunning = true;
-				return this;
-			}
-
-			virtual void Tick() override;
-
-			void PushFrame(TSharedPtr<FVideoSourceGroup> VideoSourceGroup);
-
-			bool						bIsRunning = false;
-			TWeakPtr<FVideoSourceGroup> OuterVideoSourceGroup = nullptr;
-			uint64						LastSubmitCycles = 0;
-
-			/* Use this event to signal when we should wake and also how long we should sleep for between transmitting a frame. */
-			FEventRef FrameEvent;
-		};
-
-		bool									bRunning = false;
-		bool									bThreadRunning = false;
-		bool									bCoupleFramerate = false;
-		int32									FramesPerSecond = 30;
-		TUniquePtr<FFrameThread>				FrameRunnable;
-		FRunnableThread*						FrameThread = nullptr; // constant FPS tick thread
+		void OnWebRtcFpsChanged(IConsoleVariable* Var);
+		void OnDecoupleFramerateChanged(IConsoleVariable* Var);
+
+		bool bFPSOverride = false;
+		bool bDecoupleOverride = false;
+
+		bool  bRunning = false;
+		bool  bDecoupleFramerate = false;
+		int32 FramesPerSecond = 30;
+
 		TArray<TSharedPtr<FEpicRtcVideoSource>> VideoSources;
 
 		FDelegateHandle FrameDelegateHandle;
+		FDelegateHandle FpsDelegateHandle;
+		FDelegateHandle DecoupleDelegateHandle;
+
+		TSharedPtr<FRunnableThread> FrameThread = nullptr;	 // constant FPS tick thread
+		TSharedPtr<FFrameRunnable>	FrameRunnable = nullptr; // constant fps runnable
 
 		mutable FCriticalSection CriticalSection;
 	};
diff --git a/Engine/Plugins/Media/PixelStreaming2/Source/PixelStreaming2Core/Internal/PixelStreaming2PluginSettings.h b/Engine/Plugins/Media/PixelStreaming2/Source/PixelStreaming2Core/Internal/PixelStreaming2PluginSettings.h
--- a/Engine/Plugins/Media/PixelStreaming2/Source/PixelStreaming2Core/Internal/PixelStreaming2PluginSettings.h
+++ b/Engine/Plugins/Media/PixelStreaming2/Source/PixelStreaming2Core/Internal/PixelStreaming2PluginSettings.h
@@ -87,7 +87,7 @@ public:
 		DisplayName = "Disable Latency Tester",
 		ToolTip = "If true disables latency tester being triggerable."
 		))
-	bool DisableLatencyTester = true; // TODO (Migration) RTCP-6453 make false by default when stats are back in
+	bool DisableLatencyTester = false;
 
 	static TAutoConsoleVariable<FString> CVarInputController;
 	UPROPERTY(config, EditAnywhere, Category = "PixelStreaming", meta = (
@@ -111,7 +111,7 @@ public:
 		DisplayName = "Decouple Wait Factor",
 		ToolTip = "Frame rate factor to wait for a captured frame when streaming in decoupled mode. Higher factor waits longer but may also result in higher latency."
 		))
-	float DecoupleWaitFactor = 1.25f;
+	float DecoupleWaitFactor = 1.0f;
 
 	static TAutoConsoleVariable<float> CVarSignalingReconnectInterval;
 	UPROPERTY(config, EditAnywhere, Category = "PixelStreaming", meta = (
@@ -627,6 +627,9 @@ public:
 
 		DECLARE_TS_MULTICAST_DELEGATE_OneParam(FOnEncoderDebugDumpFrameChanged, IConsoleVariable*);
 		FOnEncoderDebugDumpFrameChanged OnEncoderDebugDumpFrameChanged;
+
+		DECLARE_TS_MULTICAST_DELEGATE_OneParam(FOnDecoupleFramerateChanged, IConsoleVariable*);
+		FOnDecoupleFramerateChanged OnDecoupleFramerateChanged;
 	};
 
 	static FDelegates* Delegates();
diff --git a/Engine/Plugins/Media/PixelStreaming2/Source/PixelStreaming2Core/Private/PixelStreaming2PluginSettings.cpp b/Engine/Plugins/Media/PixelStreaming2/Source/PixelStreaming2Core/Private/PixelStreaming2PluginSettings.cpp
--- a/Engine/Plugins/Media/PixelStreaming2/Source/PixelStreaming2Core/Private/PixelStreaming2PluginSettings.cpp
+++ b/Engine/Plugins/Media/PixelStreaming2/Source/PixelStreaming2Core/Private/PixelStreaming2PluginSettings.cpp
@@ -163,6 +163,7 @@ TAutoConsoleVariable<bool> UPixelStreaming2PluginSettings::CVarDecoupleFramerate
 	TEXT("PixelStreaming2.DecoupleFramerate"),
 	false,
 	TEXT("Whether we should only stream as fast as we render or at some fixed interval. Coupled means only stream what we render."),
+	FConsoleVariableDelegate::CreateLambda([](IConsoleVariable* Var) { Delegates()->OnDecoupleFramerateChanged.Broadcast(Var); }),
 	ECVF_Default);
 
 TAutoConsoleVariable<float> UPixelStreaming2PluginSettings::CVarDecoupleWaitFactor(
