Add: CVar to dump the PS2 mixed audio

diff --git a/Engine/Plugins/Media/PixelStreaming2/Source/PixelStreaming2/Private/EpicRtcAudioMixingCapturer.cpp b/Engine/Plugins/Media/PixelStreaming2/Source/PixelStreaming2/Private/EpicRtcAudioMixingCapturer.cpp
--- a/Engine/Plugins/Media/PixelStreaming2/Source/PixelStreaming2/Private/EpicRtcAudioMixingCapturer.cpp
+++ b/Engine/Plugins/Media/PixelStreaming2/Source/PixelStreaming2/Private/EpicRtcAudioMixingCapturer.cpp
@@ -4,7 +4,10 @@
 
 #include "AudioDevice.h"
 #include "Engine/Engine.h"
+#include "Logging.h"
+#include "Misc/CoreDelegates.h"
 #include "PixelStreaming2PluginSettings.h"
+#include "Sound/SampleBufferIO.h"
 
 namespace UE::PixelStreaming2
 {
@@ -39,6 +42,19 @@ namespace UE::PixelStreaming2
 		FAudioDeviceManagerDelegates::OnAudioDeviceCreated.AddSP(AudioMixingCapturer.ToSharedRef(), &FEpicRtcAudioMixingCapturer::CreateAudioProducer);
 		FAudioDeviceManagerDelegates::OnAudioDeviceDestroyed.AddSP(AudioMixingCapturer.ToSharedRef(), &FEpicRtcAudioMixingCapturer::RemoveAudioProducer);
 
+		if (UPixelStreaming2PluginSettings::FDelegates* Delegates = UPixelStreaming2PluginSettings::Delegates())
+		{
+			Delegates->OnDebugDumpAudioChanged.AddSP(AudioMixingCapturer.ToSharedRef(), &FEpicRtcAudioMixingCapturer::OnDebugDumpAudioChanged);
+
+			TWeakPtr<FEpicRtcAudioMixingCapturer> WeakAudioMixingCapturer = AudioMixingCapturer;
+			FCoreDelegates::OnEnginePreExit.AddLambda([WeakAudioMixingCapturer]() {
+				if (TSharedPtr<FEpicRtcAudioMixingCapturer> AudioMixingCapturer = WeakAudioMixingCapturer.Pin())
+				{
+					AudioMixingCapturer->OnEnginePreExit();
+				}
+			});
+		}
+
 		return AudioMixingCapturer;
 	}
 
@@ -111,7 +127,12 @@ namespace UE::PixelStreaming2
 
 		RecordingBuffer.Append(Buffer.GetData(), Buffer.GetNumSamples());
 
-		const int32 SamplesPer10Ms = NumChannels * SampleRate * 0.01f;
+		if (UPixelStreaming2PluginSettings::CVarDebugDumpAudio.GetValueOnAnyThread())
+		{
+			DebugDumpAudioBuffer.Append(Buffer.GetData(), Buffer.GetNumSamples(), Buffer.GetNumChannels(), Buffer.GetSampleRate());
+		}
+
+		const int32	 SamplesPer10Ms = NumChannels * SampleRate * 0.01f;
 		const size_t BytesPerFrame = NumChannels * sizeof(int16_t);
 
 		// Feed in 10ms chunks
@@ -123,4 +144,36 @@ namespace UE::PixelStreaming2
 			RecordingBuffer.RemoveAt(0, SamplesPer10Ms, EAllowShrinking::No);
 		}
 	}
-} // namespace UE::PixelStreaming2 
\ No newline at end of file
+
+	void FEpicRtcAudioMixingCapturer::OnDebugDumpAudioChanged(IConsoleVariable* Var)
+	{
+		if (!Var->GetBool())
+		{
+			WriteDebugAudio();
+		}
+	}
+
+	void FEpicRtcAudioMixingCapturer::OnEnginePreExit()
+	{
+		// If engine is exiting but the dump cvar is true, we need to manually trigger a write
+		if (UPixelStreaming2PluginSettings::CVarDebugDumpAudio.GetValueOnAnyThread())
+		{
+			WriteDebugAudio();
+		}
+	}
+
+	void FEpicRtcAudioMixingCapturer::WriteDebugAudio()
+	{
+		// Only write audio if we actually have some
+		if (DebugDumpAudioBuffer.GetSampleDuration() <= 0.f)
+		{
+			return;
+		}
+
+		Audio::FSoundWavePCMWriter Writer;
+		FString					   FilePath = TEXT("");
+		Writer.SynchronouslyWriteToWavFile(DebugDumpAudioBuffer, TEXT("PixelStreamingMixedAudio"), TEXT(""), &FilePath);
+		UE_LOGFMT(LogPixelStreaming2, Log, "Saving audio sample to: {0}", FilePath);
+		DebugDumpAudioBuffer.Reset();
+	}
+} // namespace UE::PixelStreaming2
\ No newline at end of file
diff --git a/Engine/Plugins/Media/PixelStreaming2/Source/PixelStreaming2/Private/EpicRtcAudioMixingCapturer.h b/Engine/Plugins/Media/PixelStreaming2/Source/PixelStreaming2/Private/EpicRtcAudioMixingCapturer.h
--- a/Engine/Plugins/Media/PixelStreaming2/Source/PixelStreaming2/Private/EpicRtcAudioMixingCapturer.h
+++ b/Engine/Plugins/Media/PixelStreaming2/Source/PixelStreaming2/Private/EpicRtcAudioMixingCapturer.h
@@ -151,6 +151,11 @@ namespace UE::PixelStreaming2
 	protected:
 		FEpicRtcAudioMixingCapturer();
 
+	private:
+		void OnDebugDumpAudioChanged(IConsoleVariable* Var);
+		void OnEnginePreExit();
+		void WriteDebugAudio();
+
 	private:
 		TSharedPtr<FEpicRtcAudioPatchMixer>	 Mixer;
 		TUniqueTaskPtr<FEpicRtcMixAudioTask> MixerTask;
@@ -163,5 +168,7 @@ namespace UE::PixelStreaming2
 		int				  NumChannels;
 		float			  SampleSizeSeconds;
 		Audio::FResampler Resampler;
+
+		Audio::TSampleBuffer<int16_t> DebugDumpAudioBuffer;
 	};
 } // namespace UE::PixelStreaming2
\ No newline at end of file
diff --git a/Engine/Plugins/Media/PixelStreaming2/Source/PixelStreaming2Core/Internal/PixelStreaming2PluginSettings.h b/Engine/Plugins/Media/PixelStreaming2/Source/PixelStreaming2Core/Internal/PixelStreaming2PluginSettings.h
--- a/Engine/Plugins/Media/PixelStreaming2/Source/PixelStreaming2Core/Internal/PixelStreaming2PluginSettings.h
+++ b/Engine/Plugins/Media/PixelStreaming2/Source/PixelStreaming2Core/Internal/PixelStreaming2PluginSettings.h
@@ -153,6 +153,8 @@ public:
 		))
 	FString SignallingURL;
 
+	static TAutoConsoleVariable<bool> CVarDebugDumpAudio;
+
 	// Begin Cursor Settings
 	/**
 	 * Pixel streaming always requires various software cursors so they will be
@@ -628,6 +630,9 @@ public:
 		DECLARE_TS_MULTICAST_DELEGATE_OneParam(FOnEncoderDebugDumpFrameChanged, IConsoleVariable*);
 		FOnEncoderDebugDumpFrameChanged OnEncoderDebugDumpFrameChanged;
 
+		DECLARE_TS_MULTICAST_DELEGATE_OneParam(FOnDebugDumpAudioChanged, IConsoleVariable*);
+		FOnDebugDumpAudioChanged OnDebugDumpAudioChanged;
+
 		DECLARE_TS_MULTICAST_DELEGATE_OneParam(FOnDecoupleFramerateChanged, IConsoleVariable*);
 		FOnDecoupleFramerateChanged OnDecoupleFramerateChanged;
 	};
diff --git a/Engine/Plugins/Media/PixelStreaming2/Source/PixelStreaming2Core/Private/PixelStreaming2PluginSettings.cpp b/Engine/Plugins/Media/PixelStreaming2/Source/PixelStreaming2Core/Private/PixelStreaming2PluginSettings.cpp
--- a/Engine/Plugins/Media/PixelStreaming2/Source/PixelStreaming2Core/Private/PixelStreaming2PluginSettings.cpp
+++ b/Engine/Plugins/Media/PixelStreaming2/Source/PixelStreaming2Core/Private/PixelStreaming2PluginSettings.cpp
@@ -202,6 +202,13 @@ TAutoConsoleVariable<FString> UPixelStreaming2PluginSettings::CVarSignallingURL(
 	TEXT("Default URL to connect to for signalling."),
 	ECVF_Default);
 
+TAutoConsoleVariable<bool> UPixelStreaming2PluginSettings::CVarDebugDumpAudio(
+	TEXT("PixelStreaming2.DumpDebugAudio"),
+	false,
+	TEXT("Dumps mixed audio from PS2 to a file on disk for debugging purposes."),
+	FConsoleVariableDelegate::CreateLambda([](IConsoleVariable* Var) { Delegates()->OnDebugDumpAudioChanged.Broadcast(Var); }),
+	ECVF_Default);
+
 // Begin Encoder CVars
 
 TAutoConsoleVariable<int32> UPixelStreaming2PluginSettings::CVarEncoderTargetBitrate(
