[PS2] Fix Crash in VideoSourceGroup.

FrameRunnable was destroyed while FrameThread wanted to acces it. This fix moves the ret until after FrameThread is reset.

Change bIsRunning and bIsEnabled to atomics. Prevent final frame from being pushed when bIsRunning is false and the Wait returns.

Fix race condition where Stop could be called from multiple threads causing a crash. Fix was to lock with mutex.

diff --git a/Engine/Plugins/Media/PixelStreaming2/Source/PixelStreaming2/Private/VideoSourceGroup.cpp b/Engine/Plugins/Media/PixelStreaming2/Source/PixelStreaming2/Private/VideoSourceGroup.cpp
--- a/Engine/Plugins/Media/PixelStreaming2/Source/PixelStreaming2/Private/VideoSourceGroup.cpp
+++ b/Engine/Plugins/Media/PixelStreaming2/Source/PixelStreaming2/Private/VideoSourceGroup.cpp
@@ -49,9 +49,9 @@ namespace UE::PixelStreaming2
 					FrameEvent->Wait();
 				}
 
-				if (TSharedPtr<FVideoSourceGroup> VideoSourceGroup = OuterVideoSourceGroup.Pin())
+				// Need to check bIsRunning in order to not push a frame when running is disabled and FrameEvent->Wait() has returned.
+				if (TSharedPtr<FVideoSourceGroup> VideoSourceGroup = OuterVideoSourceGroup.Pin(); bIsRunning && VideoSourceGroup)
 				{
-
 					const double TimeSinceLastSubmitMs = FPlatformTime::ToMilliseconds64(FPlatformTime::Cycles64() - LastSubmitCycles);
 
 					// Decrease this value to make expected frame delivery more precise, however may result in more old frames being sent
@@ -91,11 +91,13 @@ namespace UE::PixelStreaming2
 		virtual void Stop() override
 		{
 			bIsRunning = false;
+			// Wake the thread in case it was sleeping. This will cause it to then exit the ::Run loop
+			FrameEvent->Trigger();
 		}
 
 		virtual void Exit() override
 		{
-			bIsRunning = false;
+			Stop();
 		}
 
 		virtual FSingleThreadRunnable* GetSingleThreadInterface() override
@@ -128,8 +130,8 @@ namespace UE::PixelStreaming2
 			LastSubmitCycles = FPlatformTime::Cycles64();
 		}
 
-		bool						bEnabled = false;
-		bool						bIsRunning = false;
+		std::atomic<bool>			bEnabled = false;
+		std::atomic<bool>			bIsRunning = false;
 		TWeakPtr<FVideoSourceGroup> OuterVideoSourceGroup = nullptr;
 		uint64						LastSubmitCycles = 0;
 
@@ -284,6 +286,7 @@ namespace UE::PixelStreaming2
 
 	void FVideoSourceGroup::StartThread()
 	{
+		FScopeLock Lock(&RunnableCriticalSection);
 		if (!FrameRunnable)
 		{
 			FrameRunnable = MakeShared<FFrameRunnable>(AsWeak(), bDecoupleFramerate);
@@ -297,12 +300,10 @@ namespace UE::PixelStreaming2
 
 	void FVideoSourceGroup::StopThread()
 	{
+		FScopeLock Lock(&RunnableCriticalSection);
 		if (FrameRunnable)
 		{
 			FrameRunnable->Stop();
-			// Wake the thread in case it was sleeping. This will cause it to then exit the ::Run loop
-			FrameRunnable->FrameEvent->Trigger();
-			FrameRunnable.Reset();
 		}
 		
 		if (FrameThread)
@@ -310,6 +311,12 @@ namespace UE::PixelStreaming2
 			FrameThread->Kill(true);
 			FrameThread.Reset();
 		}
+
+		// Need to reset FrameRunnable after FrameThread as that will try to access FrameRunnable when killing it.
+		if (FrameRunnable)
+		{
+			FrameRunnable.Reset();
+		}
 	}
 
 	void FVideoSourceGroup::ForceKeyFrame()
diff --git a/Engine/Plugins/Media/PixelStreaming2/Source/PixelStreaming2/Private/VideoSourceGroup.h b/Engine/Plugins/Media/PixelStreaming2/Source/PixelStreaming2/Private/VideoSourceGroup.h
--- a/Engine/Plugins/Media/PixelStreaming2/Source/PixelStreaming2/Private/VideoSourceGroup.h
+++ b/Engine/Plugins/Media/PixelStreaming2/Source/PixelStreaming2/Private/VideoSourceGroup.h
@@ -64,5 +64,6 @@ namespace UE::PixelStreaming2
 		TSharedPtr<FFrameRunnable>	FrameRunnable = nullptr; // constant fps runnable
 
 		mutable FCriticalSection CriticalSection;
+		mutable FCriticalSection RunnableCriticalSection;
 	};
 } // namespace UE::PixelStreaming2
