[PS2] Fix: Mac not working with VPX (doh!)

diff --git a/Engine/Plugins/Media/PixelStreaming2/PixelStreaming2.uplugin b/Engine/Plugins/Media/PixelStreaming2/PixelStreaming2.uplugin
--- a/Engine/Plugins/Media/PixelStreaming2/PixelStreaming2.uplugin
+++ b/Engine/Plugins/Media/PixelStreaming2/PixelStreaming2.uplugin
@@ -90,7 +90,7 @@
 		{
 			"Name": "LibVpxCodecs",
 			"Enabled": true,
-			"PlatformAllowList": [ "Win64", "Linux" ]
+			"PlatformAllowList": [ "Win64", "Linux", "Mac" ]
 		},
 		{
 			"Name": "WebSocketNetworking",
