[PS1/PS2] Fix: StreamerComponent not working with internal signalling server.

diff --git a/Engine/Plugins/Media/PixelStreaming/Source/PixelStreamingBlueprint/Private/PixelStreamingStreamerComponent.h b/Engine/Plugins/Media/PixelStreaming/Source/PixelStreamingBlueprint/Private/PixelStreamingStreamerComponent.h
--- a/Engine/Plugins/Media/PixelStreaming/Source/PixelStreamingBlueprint/Private/PixelStreamingStreamerComponent.h
+++ b/Engine/Plugins/Media/PixelStreaming/Source/PixelStreamingBlueprint/Private/PixelStreamingStreamerComponent.h
@@ -56,7 +56,7 @@ public:
 	FString StreamerId = "Streamer Component";
 
 	UPROPERTY(EditAnywhere, Category = "PixelStreaming")
-	FString SignallingServerURL = "ws://localhost:8888";
+	FString SignallingServerURL = "ws://127.0.0.1:8888";
 
 	UPROPERTY(EditAnywhere, Category = "PixelStreaming")
 	bool UsePixelStreamingURL = false;
diff --git a/Engine/Plugins/Media/PixelStreaming2/Source/PixelStreaming2/Private/Blueprints/PixelStreaming2StreamerComponent.h b/Engine/Plugins/Media/PixelStreaming2/Source/PixelStreaming2/Private/Blueprints/PixelStreaming2StreamerComponent.h
--- a/Engine/Plugins/Media/PixelStreaming2/Source/PixelStreaming2/Private/Blueprints/PixelStreaming2StreamerComponent.h
+++ b/Engine/Plugins/Media/PixelStreaming2/Source/PixelStreaming2/Private/Blueprints/PixelStreaming2StreamerComponent.h
@@ -57,7 +57,7 @@ public:
 	FString StreamerId = "Streamer Component";
 
 	UPROPERTY(EditAnywhere, Category = "PixelStreaming2")
-	FString SignallingServerURL = "ws://localhost:8888";
+	FString SignallingServerURL = "ws://127.0.0.1:8888";
 
 	UPROPERTY(EditAnywhere, Category = "PixelStreaming2")
 	bool UsePixelStreamingURL = false;
