[PS2] Update doc strings in Pixel Streaming 2

diff --git a/Engine/Plugins/Media/PixelStreaming2/Source/PixelStreaming2/Public/IPixelStreaming2AudioConsumer.h b/Engine/Plugins/Media/PixelStreaming2/Source/PixelStreaming2/Public/IPixelStreaming2AudioConsumer.h
--- a/Engine/Plugins/Media/PixelStreaming2/Source/PixelStreaming2/Public/IPixelStreaming2AudioConsumer.h
+++ b/Engine/Plugins/Media/PixelStreaming2/Source/PixelStreaming2/Public/IPixelStreaming2AudioConsumer.h
@@ -13,7 +13,22 @@ class PIXELSTREAMING2_API IPixelStreaming2AudioConsumer
 public:
 	virtual ~IPixelStreaming2AudioConsumer() = default;
 
+	/**
+	 * @brief Consume raw audio data.
+	 * @param AudioData Pointer to the audio data.
+	 * @param InSampleRate Audio sample rate in samples per second.
+	 * @param NChannels Number of audio channels. For example 2 for stero audio.
+	 * @param NFrames Number of Audio frames in a single channel.
+	 */
 	virtual void ConsumeRawPCM(const int16_t* AudioData, int InSampleRate, size_t NChannels, size_t NFrames) = 0;
+
+	/**
+	 * @brief Called when a audio consumer is added.
+	 */
 	virtual void OnConsumerAdded() = 0;
+
+	/**
+	 * @brief Called when a audio consumer is removed.
+	 */
 	virtual void OnConsumerRemoved() = 0;
 };
diff --git a/Engine/Plugins/Media/PixelStreaming2/Source/PixelStreaming2/Public/IPixelStreaming2AudioProducer.h b/Engine/Plugins/Media/PixelStreaming2/Source/PixelStreaming2/Public/IPixelStreaming2AudioProducer.h
--- a/Engine/Plugins/Media/PixelStreaming2/Source/PixelStreaming2/Public/IPixelStreaming2AudioProducer.h
+++ b/Engine/Plugins/Media/PixelStreaming2/Source/PixelStreaming2/Public/IPixelStreaming2AudioProducer.h
@@ -19,6 +19,12 @@ public:
 	IPixelStreaming2AudioProducer() = default;
 	virtual ~IPixelStreaming2AudioProducer() = default;
 
-    /* Pushes audio into the PS2 audio pipeline which will mix with other audio producers before broadcasting. */
+	/**
+	 * @brief Pushes audio into the PS2 audio pipeline which will mix with other audio producers before broadcasting.
+	 * @param InBuffer Pointer to the audio data.
+	 * @param NumSamples Number of Audio frames in a single channel.
+	 * @param NumChannels Number of audio channels. For example 2 for stero audio.
+	 * @param SampleRate Audio sample rate in samples per second.
+	 */
     virtual void PushAudio(const float* InBuffer, int32 NumSamples, int32 NumChannels, int32 SampleRate) = 0;
 };
diff --git a/Engine/Plugins/Media/PixelStreaming2/Source/PixelStreaming2/Public/IPixelStreaming2AudioSink.h b/Engine/Plugins/Media/PixelStreaming2/Source/PixelStreaming2/Public/IPixelStreaming2AudioSink.h
--- a/Engine/Plugins/Media/PixelStreaming2/Source/PixelStreaming2/Public/IPixelStreaming2AudioSink.h
+++ b/Engine/Plugins/Media/PixelStreaming2/Source/PixelStreaming2/Public/IPixelStreaming2AudioSink.h
@@ -15,6 +15,15 @@ class PIXELSTREAMING2_API IPixelStreaming2AudioSink
 public:
 	virtual ~IPixelStreaming2AudioSink() = default;
 
+	/**
+	 * @brief Add an audio consumer to the sink.
+	 * @param AudioConsumer The Audio consumer to add to the sink.
+	 */
 	virtual void AddAudioConsumer(IPixelStreaming2AudioConsumer* AudioConsumer) = 0;
+
+	/**
+	 * @brief Remove an audio consumer to remove from the sink.
+	 * @param AudioConsumer The Audio consumer to remove from the sink.
+	 */
 	virtual void RemoveAudioConsumer(IPixelStreaming2AudioConsumer* AudioConsumer) = 0;		
 };
\ No newline at end of file
diff --git a/Engine/Plugins/Media/PixelStreaming2/Source/PixelStreaming2/Public/IPixelStreaming2Module.h b/Engine/Plugins/Media/PixelStreaming2/Source/PixelStreaming2/Public/IPixelStreaming2Module.h
--- a/Engine/Plugins/Media/PixelStreaming2/Source/PixelStreaming2/Public/IPixelStreaming2Module.h
+++ b/Engine/Plugins/Media/PixelStreaming2/Source/PixelStreaming2/Public/IPixelStreaming2Module.h
@@ -13,14 +13,16 @@
 #include "Video/VideoConfig.h"
 
 /**
- * The public interface of the Pixel Streaming module.
+ * The IPixelStreaming2Module interface manages the core functionality of the Pixel Streaming system.
+ * This class provides access to streamers, video and audio producers, and handles starting and stopping the streaming process.
+ * It also allows interaction with the signalling server and manages the lifecycle of streamers within the Pixel Streaming system.
  */
 class PIXELSTREAMING2_API IPixelStreaming2Module : public IModuleInterface
 {
 public:
 	/**
 	 * Singleton-like access to this module's interface.
-	 * Beware of calling this during the shutdown phase, though.  Your module might have been unloaded already.
+	 * Beware calling this during the shutdown phase, though. Your module might have been unloaded already.
 	 *
 	 * @return Returns singleton instance, loading the module on demand if needed
 	 */
@@ -58,6 +60,7 @@ public:
 
 	/**
 	 * Starts streaming on all streamers.
+	 * @return False if the module was not able to start streaming. 
 	 */
 	virtual bool StartStreaming() = 0;
 
@@ -69,19 +72,20 @@ public:
 	/**
 	 * Creates a new streamer.
 	 * @param StreamerId - The ID of the Streamer to be created.
+	 * @return SharedPtr to the streamer.
 	 */
 	virtual TSharedPtr<IPixelStreaming2Streamer> CreateStreamer(const FString& StreamerId) = 0;
 
 	/**
 	 * @brief Creates a new video producer.
-	 * @return An video producer interface for you to push video frames in to
+	 * @return An video producer interface for you to push video frames in to.
 	 */
 	virtual TSharedPtr<IPixelStreaming2VideoProducer> CreateVideoProducer() = 0;
 
 	/**
-	 * @brief Creates a new audio producer. Any audio you push in with `PushAudio` will be mixed with other audio sources before being streamed
-	 * @note Users are responsible for the lifetime of this object
-	 * @return An audio producer interface for you to push audio in to
+	 * @brief Creates a new audio producer. Any audio you push in with `PushAudio` will be mixed with other audio sources before being streamed.
+	 * @note Users are responsible for the lifetime of this object.
+	 * @return An audio producer interface for you to push audio in to.
 	 */
 	virtual TSharedPtr<IPixelStreaming2AudioProducer> CreateAudioProducer() = 0;
 
@@ -93,7 +97,8 @@ public:
 
 	/**
 	 * Find a streamer by an ID.
-	 * @return A pointer to the interface for a streamer. nullptr if the streamer isn't found
+	 * @param StreamerId	-	The ID of the streamer to be found.
+	 * @return A pointer to the interface for a streamer. nullptr if the streamer isn't found.
 	 */
 	virtual TSharedPtr<IPixelStreaming2Streamer> FindStreamer(const FString& StreamerId) = 0;
 
@@ -111,21 +116,21 @@ public:
 	virtual void DeleteStreamer(TSharedPtr<IPixelStreaming2Streamer> ToBeDeleted) = 0;
 
 	/**
-	 * Get the Default Streamer ID
-	 * @return FString The default streamer ID
+	 * Get the Default Streamer ID.
+	 * @return FString The default streamer ID.
 	 */
 	virtual FString GetDefaultStreamerID() = 0;
 
 	/**
-	 * Get the Default Signaling URL ("ws://127.0.0.1:8888")
-	 * @return FString The default signaling url ("ws://127.0.0.1:8888")
+	 * Get the Default Signaling URL ("ws://127.0.0.1:8888").
+	 * @return FString The default signaling url ("ws://127.0.0.1:8888").
 	 */
 	virtual FString GetDefaultSignallingURL() = 0;
 
 	/**
-	 * @brief A method for iterating through all of the streamers on the module
+	 * @brief A method for iterating through all of the streamers on the module.
 	 *
-	 * @param Func The lambda to execute with each streamer
+	 * @param Func The lambda to execute with each streamer.
 	 */
 	virtual void ForEachStreamer(const TFunction<void(TSharedPtr<IPixelStreaming2Streamer>)>& Func) = 0;
 };
diff --git a/Engine/Plugins/Media/PixelStreaming2/Source/PixelStreaming2/Public/IPixelStreaming2Streamer.h b/Engine/Plugins/Media/PixelStreaming2/Source/PixelStreaming2/Public/IPixelStreaming2Streamer.h
--- a/Engine/Plugins/Media/PixelStreaming2/Source/PixelStreaming2/Public/IPixelStreaming2Streamer.h
+++ b/Engine/Plugins/Media/PixelStreaming2/Source/PixelStreaming2/Public/IPixelStreaming2Streamer.h
@@ -14,20 +14,25 @@
 
 class UTexture2D;
 
+/**
+ * The IPixelStreaming2Streamer interface provides functionality for managing a Pixel Streaming session.
+ * This class manages the core operations of streaming, such as setting the framerate, handling video input, 
+ * and interacting with the signalling server. It also allows for control over connected players and sending data to them.
+ */
 class PIXELSTREAMING2_API IPixelStreaming2Streamer
 {
 public:
 	virtual ~IPixelStreaming2Streamer() = default;
 
 	/**
-	 * @brief Set the Stream FPS
-	 * @param InFramesPerSecond The number of frames per second the streamer will stream at
+	 * @brief Set the Stream FPS.
+	 * @param InFramesPerSecond The number of frames per second the streamer will stream at.
 	 */
 	virtual void SetStreamFPS(int32 InFramesPerSecond) = 0;
 
 	/**
-	 * @brief Get the Stream FPS
-	 * @return - The number of frames per second the streamer will stream at
+	 * @brief Get the Stream FPS.
+	 * @return - The number of frames per second the streamer will stream at.
 	 */
 	virtual int32 GetStreamFPS() = 0;
 
@@ -40,32 +45,32 @@ public:
 	virtual void SetCoupleFramerate(bool bCouple) = 0;
 
 	/**
-	 * @brief Set the Video Input object
-	 * @param Input The IPixelStreaming2VideoProducer that this streamer will stream
+	 * @brief Set the Video Input object.
+	 * @param Input The IPixelStreaming2VideoProducer that this streamer will stream.
 	 */
 	virtual void SetVideoProducer(TSharedPtr<IPixelStreaming2VideoProducer> Input) = 0;
 
 	/**
-	 * @brief Get the Video Input object
-	 * @return The IPixelStreaming2VideoProducer that this streamer will stream
+	 * @brief Get the Video Input object.
+	 * @return The IPixelStreaming2VideoProducer that this streamer will stream.
 	 */
 	virtual TWeakPtr<IPixelStreaming2VideoProducer> GetVideoProducer() = 0;
 
 	/**
-	 * @brief Set the Signalling Server URL
-	 * @param InSignallingServerURL
+	 * @brief Set the Signalling Server URL.
+	 * @param InSignallingServerURL.
 	 */
 	virtual void SetSignallingServerURL(const FString& InSignallingServerURL) = 0;
 
 	/**
-	 * @brief Get the Signalling Server URL
-	 * @return The Signalling Server URL
+	 * @brief Get the Signalling Server URL.
+	 * @return The Signalling Server URL.
 	 */
 	virtual FString GetSignallingServerURL() = 0;
 
 	/**
-	 * @brief Get this streamer's ID
-	 *
+	 * @brief Get this streamer's ID.
+	 * @return The streamer's ID.
 	 */
 	virtual FString GetId() = 0;
 
@@ -87,8 +92,7 @@ public:
 
 	/**
 	 * @brief Get the current state of this streamer
-	 * @return true
-	 * @return false
+	 * @return True if streaming, false otherwise.
 	 */
 	virtual bool IsStreaming() const = 0;
 
@@ -142,15 +146,15 @@ public:
 	virtual void UnfreezeStream() = 0;
 
 	/**
-	 * @brief Send all players connected to this streamer a message
-	 * @param MessageType The message type to be sent to the player
+	 * @brief Send all players connected to this streamer a message.
+	 * @param MessageType The message type to be sent to the player.
 	 * @param Descriptor The contents of the message.
 	 */
 	virtual void SendAllPlayersMessage(FString MessageType, const FString& Descriptor) = 0;
 
 	/**
-	 * @brief Send all players connected to this streamer a message
-	 * @param MessageType The message type to be sent to the player
+	 * @brief Send all players connected to this streamer a message.
+	 * @param MessageType The message type to be sent to the player.
 	 * @param Descriptor The contents of the message.
 	 */
 	virtual void SendPlayerMessage(FString PlayerId, FString MessageType, const FString& Descriptor) = 0;
@@ -171,7 +175,6 @@ public:
 
 	/**
 	 * @brief Get the ids of the connected players.
-	 *
 	 * @return TArray<FString> The ids of the connected players.
 	 */
 	virtual TArray<FString> GetConnectedPlayers() = 0;
@@ -184,35 +187,53 @@ public:
 
 	/**
 	 * @brief Get the audio sink associated with a specific peer/player.
+	 * @param PlayerId The player id of the audio sink to retrieve.
+	 * @return The PeerAudioSink.
 	 */
 	virtual IPixelStreaming2AudioSink* GetPeerAudioSink(FString PlayerId) = 0;
 
 	/**
 	 * @brief Get an audio sink that has no peers/players listening to it.
+	 * @return The unlistened PeerAudioSink.
 	 */
 	virtual IPixelStreaming2AudioSink* GetUnlistenedAudioSink() = 0;
 
 	/**
 	 * @brief Get the video sink associated with a specific peer/player.
+	 * @return The PeerVideoSink.
 	 */
 	virtual IPixelStreaming2VideoSink* GetPeerVideoSink(FString PlayerId) = 0;
 
 	/**
 	 * @brief Get a video sink that has no peers/players watching it.
+	 * @return The unwatched PeerVideoSink.
 	 */
 	virtual IPixelStreaming2VideoSink* GetUnwatchedVideoSink() = 0;
 
 	/**
-	 * @brief Allows sending arbitrary configuration options during initial connection
-	 *
-	 * @param Value Setting a value to an empty string clears it from the mapping and prevents it being sent
+	 * @brief Allows sending arbitrary configuration options during initial connection.
+	 * @param OptionName The name of the option to set.
+	 * @param Value Setting a value to an empty string clears it from the mapping and prevents it being sent.
 	 */
 	virtual void SetConfigOption(const FName& OptionName, const FString& Value) = 0;
+
+	/**
+	 * @brief Get the Gonfigerartion value for a specific option.
+	 * @param OptionName The name of the config option to get.
+	 * @param Value The value of the config option retrieved.
+	 */
 	virtual bool GetConfigOption(const FName& OptionName, FString& OutValue) = 0;
 
 	/**
+	 * @brief Set the minimum and maximum bitrate for the streamer.
+	 * @param PlayerId Currently unused. For setting the bitrate for a individual player id.
+	 * @param MinBitrate minimum bitrate for the streamer.
+	 * @param MaxBitrate maximum bitrate for the streamer.
 	 */
 	virtual void PlayerRequestsBitrate(FString PlayerId, int MinBitrate, int MaxBitrate) = 0;
 
+	/**
+	 * @brief Refresh connection with minimum and maximum bitrate.
+	 */
 	virtual void RefreshStreamBitrate() = 0;
 };
diff --git a/Engine/Plugins/Media/PixelStreaming2/Source/PixelStreaming2/Public/IPixelStreaming2VideoConsumer.h b/Engine/Plugins/Media/PixelStreaming2/Source/PixelStreaming2/Public/IPixelStreaming2VideoConsumer.h
--- a/Engine/Plugins/Media/PixelStreaming2/Source/PixelStreaming2/Public/IPixelStreaming2VideoConsumer.h
+++ b/Engine/Plugins/Media/PixelStreaming2/Source/PixelStreaming2/Public/IPixelStreaming2VideoConsumer.h
@@ -14,7 +14,19 @@ class PIXELSTREAMING2_API IPixelStreaming2VideoConsumer
 public:
 	virtual ~IPixelStreaming2VideoConsumer() = default;
 
+	/**
+	 * @brief Consume a texture as a video frame.
+	 * @param Frame The Frame to consume.
+	 */
 	virtual void ConsumeFrame(FTextureRHIRef Frame) = 0;
+
+	/**
+	 * @brief Called when a video consumer is added.
+	 */
 	virtual void OnConsumerAdded() = 0;
+
+	/**
+	 * @brief Called when a video consumer is removed.
+	 */
 	virtual void OnConsumerRemoved() = 0;
 };
diff --git a/Engine/Plugins/Media/PixelStreaming2/Source/PixelStreaming2/Public/IPixelStreaming2VideoProducer.h b/Engine/Plugins/Media/PixelStreaming2/Source/PixelStreaming2/Public/IPixelStreaming2VideoProducer.h
--- a/Engine/Plugins/Media/PixelStreaming2/Source/PixelStreaming2/Public/IPixelStreaming2VideoProducer.h
+++ b/Engine/Plugins/Media/PixelStreaming2/Source/PixelStreaming2/Public/IPixelStreaming2VideoProducer.h
@@ -29,6 +29,7 @@ public:
 
     /**
 	 * A human readable identifier used when displaying what the streamer is streaming in the toolbar
+	 * @return A string containing the display name.
 	 */
 	virtual FString ToString() = 0;
 };
\ No newline at end of file
diff --git a/Engine/Plugins/Media/PixelStreaming2/Source/PixelStreaming2/Public/IPixelStreaming2VideoSink.h b/Engine/Plugins/Media/PixelStreaming2/Source/PixelStreaming2/Public/IPixelStreaming2VideoSink.h
--- a/Engine/Plugins/Media/PixelStreaming2/Source/PixelStreaming2/Public/IPixelStreaming2VideoSink.h
+++ b/Engine/Plugins/Media/PixelStreaming2/Source/PixelStreaming2/Public/IPixelStreaming2VideoSink.h
@@ -6,12 +6,23 @@
 
 class IPixelStreaming2VideoConsumer;
 
-// Interface for a sink that collects video coming in from the browser and passes into into UE.
+/*
+ * Interface for a sink that collects video coming in from the browser and passes into into UE.
+ */
 class PIXELSTREAMING2_API IPixelStreaming2VideoSink
 {
 public:
 	virtual ~IPixelStreaming2VideoSink() = default;
 
+	/**
+	 * @brief Add a video consumer to the sink.
+	 * @param VideoConsumer The video consumer to add to the sink.
+	 */
 	virtual void AddVideoConsumer(IPixelStreaming2VideoConsumer* VideoConsumer) = 0;
+
+	/**
+	 * @brief Remove a video consumer to remove from the sink.
+	 * @param VideoConsumer The video consumer to remove from the sink.
+	 */
 	virtual void RemoveVideoConsumer(IPixelStreaming2VideoConsumer* VideoConsumer) = 0;
 };
diff --git a/Engine/Plugins/Media/PixelStreaming2/Source/PixelStreaming2/Public/PixelStreaming2Delegates.h b/Engine/Plugins/Media/PixelStreaming2/Source/PixelStreaming2/Public/PixelStreaming2Delegates.h
--- a/Engine/Plugins/Media/PixelStreaming2/Source/PixelStreaming2/Public/PixelStreaming2Delegates.h
+++ b/Engine/Plugins/Media/PixelStreaming2/Source/PixelStreaming2/Public/PixelStreaming2Delegates.h
@@ -6,142 +6,160 @@
 #include "Containers/UnrealString.h"
 #include "PixelStreaming2Delegates.generated.h"
 
+/**
+ * Pixel Streaming Delegates that can be invoked when pixel streaming events take place.
+ * Includes blueprint and native c++ delegates.
+ */
 UCLASS()
 class PIXELSTREAMING2_API UPixelStreaming2Delegates : public UObject
 {
 	GENERATED_BODY()
 
 public:
-	/**
-	 * A connection to the signalling server was made.
-	 */
-	// BP Delegate
+	/** Blueprint delegate type for when a connection to the signalling server was made. */
 	DECLARE_DYNAMIC_MULTICAST_DELEGATE_OneParam(FConnectedToSignallingServer, FString, StreamerId);
+
+	/** Invoked when a connection to the signalling server was made. */
 	UPROPERTY(BlueprintAssignable, Category = "Pixel Streaming Delegates")
 	FConnectedToSignallingServer OnConnectedToSignallingServer;
-	// C++ Delegate
+	
+	/** Delegate type for when a connection to the signalling server was made. */
 	DECLARE_MULTICAST_DELEGATE_OneParam(FConnectedToSignallingServerNative, FString /* StreamerId */);
+
+	/** Invoked when a connection to the signalling server was made. */
 	FConnectedToSignallingServerNative OnConnectedToSignallingServerNative;
 
-	/**
-	 * A connection to the signalling server was lost.
-	 */
-	// BP Delegate
+	/** Blueprint delegate type for when a connection to the signalling server was lost. */
 	DECLARE_DYNAMIC_MULTICAST_DELEGATE_OneParam(FDisconnectedFromSignallingServer, FString, StreamerId);
+
+	/** Invoked when a connection to the signalling server was lost. */
 	UPROPERTY(BlueprintAssignable, Category = "Pixel Streaming Delegates")
 	FDisconnectedFromSignallingServer OnDisconnectedFromSignallingServer;
-	// C++ Delegate
+	
+	/** Delegate type for when a connection to the signalling server was lost. */
 	DECLARE_MULTICAST_DELEGATE_OneParam(FDisconnectedFromSignallingServerNative, FString /* StreamerId */);
+
+	/** Invoked when a connection to the signalling server was lost. */
 	FDisconnectedFromSignallingServerNative OnDisconnectedFromSignallingServerNative;
 
-	/**
-	 * A new connection has been made to the session.
-	 */
-	// BP Delegate
+	/** Blueprint delegate type for when a new connection has been made to the session. */
 	DECLARE_DYNAMIC_MULTICAST_DELEGATE_TwoParams(FNewConnection, FString, StreamerId, FString, PlayerId);
+
+	/** Invoked when a new connection has been made to the session. */
 	UPROPERTY(BlueprintAssignable, Category = "Pixel Streaming Delegates")
 	FNewConnection OnNewConnection;
-	// C++ Delegate
+
+	/** Delegate type for when a new connection has been made to the session. */
 	DECLARE_TS_MULTICAST_DELEGATE_TwoParams(FNewConnectionNative, FString /* Streamer id*/, FString /* Peer id */);
+
+	/** Invoked when a new connection has been made to the session. */
 	FNewConnectionNative OnNewConnectionNative;
 
-	/**
-	 * A connection to a player was lost.
-	 */
-	// BP Delegate
+	/** Blueprint delegate type for when a connection to a player was lost. */
 	DECLARE_DYNAMIC_MULTICAST_DELEGATE_TwoParams(FClosedConnection, FString, StreamerId, FString, PlayerId);
+
+	/** Invoked when a connection to a player was lost. */
 	UPROPERTY(BlueprintAssignable, Category = "Pixel Streaming Delegates")
 	FClosedConnection OnClosedConnection;
-	// C++ Delegate
+
+	/** Delegate type for when a connection to a player was lost. */
 	DECLARE_TS_MULTICAST_DELEGATE_TwoParams(FClosedConnectionNative, FString /* Streamer id */, FString /* Peer id */);
+
+	/** Invoked when a connection to a player was lost. */
 	FClosedConnectionNative OnClosedConnectionNative;
 
-	/**
-	 * All connections have closed and nobody is viewing or interacting with
-	 * the app. This is an opportunity to reset the app.
-	 */
-	// BP Delegate
+	/** Blueprint delegate type for when all connections have closed and nobody is viewing or interacting with the app. This is an opportunity to reset the app. */
 	DECLARE_DYNAMIC_MULTICAST_DELEGATE_OneParam(FAllConnectionsClosed, FString, StreamerId);
+
+	/** Invoked when all connections have closed and nobody is viewing or interacting with the app. This is an opportunity to reset the app. */
 	UPROPERTY(BlueprintAssignable, Category = "Pixel Streaming Delegates")
 	FAllConnectionsClosed OnAllConnectionsClosed;
-	// C++ Delegate
+
+	/** Delegate type for when all connections have closed and nobody is viewing or interacting with the app. This is an opportunity to reset the app. */
 	DECLARE_TS_MULTICAST_DELEGATE_OneParam(FAllConnectionsClosedNative, FString);
+
+	/** Invoked when all connections have closed and nobody is viewing or interacting with the app. This is an opportunity to reset the app. */
 	FAllConnectionsClosedNative OnAllConnectionsClosedNative;
 
-	/**
-	 * A new data track has been opened
-	 */
-	// BP Delegate
+	/** Blueprint delegate type for when a new data track has been opened. */
 	DECLARE_DYNAMIC_MULTICAST_DELEGATE_TwoParams(FDataTrackOpen, FString, StreamerId, FString, PlayerId);
+
+	/** Invoked when a new data track has been opened. */
 	UPROPERTY(BlueprintAssignable, Category = "Pixel Streaming Delegates")
 	FDataTrackOpen OnDataTrackOpen;
-	// C++ Delegate
+	
+	/** Delegate type for when a new data track has been opened. */
 	DECLARE_TS_MULTICAST_DELEGATE_TwoParams(FDataTrackOpenNative, FString /* Streamer id */, FString /* Peer id */);
+
+	/** Invoked when a new data track has been opened. */
 	FDataTrackOpenNative OnDataTrackOpenNative;
 
-	/**
-	 * An existing data track has been closed
-	 */
-	// BP Delegate
+	/** Blueprint delegate type for when an existing data track has been closed. */
 	DECLARE_DYNAMIC_MULTICAST_DELEGATE_TwoParams(FDataTrackClosed, FString, StreamerId, FString, PlayerId);
+
+	/** Invoked when an existing data track has been closed. */
 	UPROPERTY(BlueprintAssignable, Category = "Pixel Streaming Delegates")
 	FDataTrackClosed OnDataTrackClosed;
-	// C++ Delegate
+
+	/** Delegate type for when an existing data track has been closed. */
 	DECLARE_TS_MULTICAST_DELEGATE_TwoParams(FDataTrackClosedNative, FString /* Streamer id */, FString /* Peer id */);
+
+	/** Invoked when an existing data track has been closed. */
 	FDataTrackClosedNative OnDataTrackClosedNative;
 
-	/**
-	 * A new video track has been opened
-	 */
-	// C++ Delegate
+	/** Delegate type for when a new video track has been opened. */
 	DECLARE_TS_MULTICAST_DELEGATE_ThreeParams(FVideoTrackOpenNative, FString /* Streamer id */, FString /* Peer id */, bool /* bIsRemote */);
+
+	/** Invoked when a new video track has been opened. */
 	FVideoTrackOpenNative OnVideoTrackOpenNative;
 
-	/**
-	 * An existing video track has been closed
-	 */
-	// C++ Delegate
+	/** Delegate type for when an existing video track has been closed. */
 	DECLARE_TS_MULTICAST_DELEGATE_ThreeParams(FVideoTrackClosedNative, FString /* Streamer id */, FString /* Peer id */, bool /* bIsRemote */);
+
+	/** Invoked when an existing video track has been closed. */
 	FVideoTrackClosedNative OnVideoTrackClosedNative;
 
-	/**
-	 * A new audio track has been opened
-	 */
-	// C++ Delegate
+	/** Delegate type for when a new audio track has been opened. */
 	DECLARE_TS_MULTICAST_DELEGATE_ThreeParams(FAudioTrackOpenNative, FString /* Streamer id */, FString /* Peer id */, bool /* bIsRemote */);
+
+	/** Invoked when a new audio track has been opened. */
 	FAudioTrackOpenNative OnAudioTrackOpenNative;
 
-	/**
-	 * An existing audio track has been closed
-	 */
-	// C++ Delegate
+	/** Delegate type for when an existing audio track has been closed. */
 	DECLARE_TS_MULTICAST_DELEGATE_ThreeParams(FAudioTrackClosedNative, FString /* Streamer id */, FString /* Peer id */, bool /* bIsRemote */);
+
+	/** Invoked when an existing audio track has been closed. */
 	FAudioTrackClosedNative OnAudioTrackClosedNative;
 
-	/**
-	 * A pixel streaming stat has changed
-	 */
-	// BP Delegate
+	/** Blueprint delegate type for when a pixel streaming stat has changed. */
 	DECLARE_DYNAMIC_MULTICAST_DELEGATE_ThreeParams(FStatsChanged, FString, PlayerId, FName, StatName, float, StatValue);
+
+	/** Invoked when a pixel streaming stat has changed. */
 	UPROPERTY(BlueprintAssignable, Category = "Pixel Streaming Delegates")
 	FStatsChanged OnStatChanged;
-	// C++ Delegate
+
+	/** Delegate type for when a pixel streaming stat has changed. */
 	DECLARE_MULTICAST_DELEGATE_ThreeParams(FStatChangedNative, FString /* Peer id */, FName /* Stat name */, float);
+
+	/** Invoked when a pixel streaming stat has changed. */
 	FStatChangedNative OnStatChangedNative;
 
-	/**
-	 * The GPU ran out of available HW encoders and fell back to software encoders
-	 */
-	// BP Delegate
+	/** Blueprint delegate type for when the GPU ran out of available hardware encoders and fell back to software encoders. */
 	DECLARE_DYNAMIC_MULTICAST_DELEGATE(FFallbackToSoftwareEncoding);
+
+	/** Invoked when the GPU ran out of available hardware encoders and fell back to software encoders. */
 	UPROPERTY(BlueprintAssignable, Category = "Pixel Streaming Delegates")
 	FFallbackToSoftwareEncoding OnFallbackToSoftwareEncoding;
-	// C++ Delegate
+
+	/** Delegate type for when the GPU ran out of available hardware encoders and fell back to software encoders. */
 	DECLARE_MULTICAST_DELEGATE(FFallbackToSoftwareEncodingNative);
+
+	/** Invoked when the GPU ran out of available hardware encoders and fell back to software encoders. */
 	FFallbackToSoftwareEncodingNative OnFallbackToSoftwareEncodingNative;
 
 	/**
-	 * Create the singleton.
+	 * @param Get the UPixelStreaming2Delegates singleton.
+	 * @return UPixelStreaming2Delegates pointer.
 	 */
 	static UPixelStreaming2Delegates* Get()
 	{
diff --git a/Engine/Plugins/Media/PixelStreaming2/Source/PixelStreaming2/Public/PixelStreaming2StatNames.h b/Engine/Plugins/Media/PixelStreaming2/Source/PixelStreaming2/Public/PixelStreaming2StatNames.h
--- a/Engine/Plugins/Media/PixelStreaming2/Source/PixelStreaming2/Public/PixelStreaming2StatNames.h
+++ b/Engine/Plugins/Media/PixelStreaming2/Source/PixelStreaming2/Public/PixelStreaming2StatNames.h
@@ -4,77 +4,72 @@
 
 #include "UObject/NameTypes.h"
 
+/**
+ * Stat Names used by Pixel Streaming.
+ */
 namespace PixelStreaming2StatNames
 {
-	// Calculated stats
-	const FName Bitrate					= FName(TEXT("bitrate"));
-	const FName BitrateMegabits			= FName(TEXT("bitrateMbits"));
-	const FName TargetBitrateMegabits	= FName(TEXT("targetBitrateMbits"));
-	const FName MeanSendDelay			= FName(TEXT("captureToSend"));
-	const FName SourceFps				= FName(TEXT("captureFps"));
-	const FName Fps						= FName(TEXT("captureFps"));
-	const FName MeanEncodeTime			= FName(TEXT("encodeTime"));
-	const FName EncodedFramesPerSecond	= FName(TEXT("encodeFps"));
-	const FName DecodedFramesPerSecond	= FName(TEXT("decodeFps"));
-	const FName MeanQPPerSecond			= FName(TEXT("qp"));
-	const FName FramesSentPerSecond		= FName(TEXT("transmitFps"));
-	const FName FramesReceivedPerSecond	= FName(TEXT("receiveFps"));
+	PIXELSTREAMING2_API extern const FName Bitrate;
+	PIXELSTREAMING2_API extern const FName BitrateMegabits;
+	PIXELSTREAMING2_API extern const FName TargetBitrateMegabits;
+	PIXELSTREAMING2_API extern const FName MeanSendDelay;
+	PIXELSTREAMING2_API extern const FName SourceFps;
+	PIXELSTREAMING2_API extern const FName Fps;
+	PIXELSTREAMING2_API extern const FName MeanEncodeTime;
+	PIXELSTREAMING2_API extern const FName EncodedFramesPerSecond;
+	PIXELSTREAMING2_API extern const FName DecodedFramesPerSecond;
+	PIXELSTREAMING2_API extern const FName MeanQPPerSecond;
+	PIXELSTREAMING2_API extern const FName FramesSentPerSecond;
+	PIXELSTREAMING2_API extern const FName FramesReceivedPerSecond;
+	PIXELSTREAMING2_API extern const FName JitterBufferDelay;
+	PIXELSTREAMING2_API extern const FName FramesSent;
+	PIXELSTREAMING2_API extern const FName FramesReceived;
+	PIXELSTREAMING2_API extern const FName FramesPerSecond;
+	PIXELSTREAMING2_API extern const FName FramesDecoded;
+	PIXELSTREAMING2_API extern const FName FramesDropped;
+	PIXELSTREAMING2_API extern const FName FramesCorrupted;
+	PIXELSTREAMING2_API extern const FName PartialFramesLost;
+	PIXELSTREAMING2_API extern const FName FullFramesLost;
+	PIXELSTREAMING2_API extern const FName HugeFramesSent;
+	PIXELSTREAMING2_API extern const FName JitterBufferTargetDelay;
+	PIXELSTREAMING2_API extern const FName InterruptionCount;
+	PIXELSTREAMING2_API extern const FName TotalInterruptionDuration;
+	PIXELSTREAMING2_API extern const FName FreezeCount;
+	PIXELSTREAMING2_API extern const FName PauseCount;
+	PIXELSTREAMING2_API extern const FName TotalFreezesDuration;
+	PIXELSTREAMING2_API extern const FName TotalPausesDuration;
+	PIXELSTREAMING2_API extern const FName FirCount;
+	PIXELSTREAMING2_API extern const FName PliCount;
+	PIXELSTREAMING2_API extern const FName NackCount;
+	PIXELSTREAMING2_API extern const FName RetransmittedBytesSent;
+	PIXELSTREAMING2_API extern const FName TargetBitrate;
+	PIXELSTREAMING2_API extern const FName TotalEncodeBytesTarget;
+	PIXELSTREAMING2_API extern const FName KeyFramesEncoded;
+	PIXELSTREAMING2_API extern const FName FrameWidth;
+	PIXELSTREAMING2_API extern const FName FrameHeight;
+	PIXELSTREAMING2_API extern const FName BytesSent;
+	PIXELSTREAMING2_API extern const FName BytesReceived;
+	PIXELSTREAMING2_API extern const FName QPSum;
+	PIXELSTREAMING2_API extern const FName TotalEncodeTime;
+	PIXELSTREAMING2_API extern const FName TotalPacketSendDelay;
+	PIXELSTREAMING2_API extern const FName FramesEncoded;
+	PIXELSTREAMING2_API extern const FName AvgSendDelay;
+	PIXELSTREAMING2_API extern const FName MessagesSent;
+	PIXELSTREAMING2_API extern const FName MessagesReceived;
+	PIXELSTREAMING2_API extern const FName PacketsLost;
+	PIXELSTREAMING2_API extern const FName Jitter;
+	PIXELSTREAMING2_API extern const FName RoundTripTime;
+	PIXELSTREAMING2_API extern const FName KeyFramesDecoded;
+	PIXELSTREAMING2_API extern const FName AudioLevel;
+	PIXELSTREAMING2_API extern const FName TotalSamplesDuration;
+	PIXELSTREAMING2_API extern const FName AvailableOutgoingBitrate;
+	PIXELSTREAMING2_API extern const FName AvailableIncomingBitrate;
+	PIXELSTREAMING2_API extern const FName RetransmittedBytesReceived;
+	PIXELSTREAMING2_API extern const FName RetransmittedPacketsReceived;
+	PIXELSTREAMING2_API extern const FName DataChannelBytesSent;
+	PIXELSTREAMING2_API extern const FName DataChannelBytesReceived;
+	PIXELSTREAMING2_API extern const FName DataChannelMessagesSent;
+	PIXELSTREAMING2_API extern const FName DataChannelMessagesReceived;
+	PIXELSTREAMING2_API extern const FName InputController;
 
-	// WebRTC stats
-	const FName JitterBufferDelay		  = FName(TEXT("jitterBufferDelay"));
-	const FName FramesSent				  = FName(TEXT("framesSent"));
-	const FName FramesReceived			  = FName(TEXT("framesReceived"));
-	const FName FramesPerSecond			  = FName(TEXT("framesPerSecond"));
-	const FName FramesDecoded			  = FName(TEXT("framesDecoded"));
-	const FName FramesDropped			  = FName(TEXT("framesDropped"));
-	const FName FramesCorrupted			  = FName(TEXT("framesCorrupted"));
-	const FName PartialFramesLost		  = FName(TEXT("partialFramesLost"));
-	const FName FullFramesLost			  = FName(TEXT("fullFramesLost"));
-	const FName HugeFramesSent			  = FName(TEXT("hugeFramesSent"));
-	const FName JitterBufferTargetDelay	  = FName(TEXT("jitterBufferTargetDelay"));
-	const FName InterruptionCount		  = FName(TEXT("interruptionCount"));
-	const FName TotalInterruptionDuration = FName(TEXT("totalInterruptionDuration"));
-	const FName FreezeCount				  = FName(TEXT("freezeCount"));
-	const FName PauseCount				  = FName(TEXT("pauseCount"));
-	const FName TotalFreezesDuration	  = FName(TEXT("totalFreezesDuration"));
-	const FName TotalPausesDuration		  = FName(TEXT("totalPausesDuration"));
-	const FName FirCount				  = FName(TEXT("firCount"));
-	const FName PliCount				  = FName(TEXT("pliCount"));
-	const FName NackCount				  = FName(TEXT("nackCount"));
-	const FName RetransmittedBytesSent	  = FName(TEXT("retransmittedBytesSent"));
-	const FName TargetBitrate			  = FName(TEXT("targetBitrate"));
-	const FName TotalEncodeBytesTarget	  = FName(TEXT("totalEncodedBytesTarget"));
-	const FName KeyFramesEncoded		  = FName(TEXT("keyFramesEncoded"));
-	const FName FrameWidth				  = FName(TEXT("frameWidth"));
-	const FName FrameHeight				  = FName(TEXT("frameHeight"));
-	const FName BytesSent				  = FName(TEXT("bytesSent"));
-	const FName BytesReceived			  = FName(TEXT("bytesReceived"));
-	const FName QPSum					  = FName(TEXT("qpSum"));
-	const FName TotalEncodeTime			  = FName(TEXT("totalEncodeTime"));
-	const FName TotalPacketSendDelay	  = FName(TEXT("totalPacketSendDelay"));
-	const FName FramesEncoded			  = FName(TEXT("framesEncoded"));
-	const FName AvgSendDelay			  = FName(TEXT("packetSendDelay"));
-	const FName MessagesSent			  = FName(TEXT("messagesSent"));
-	const FName MessagesReceived		  = FName(TEXT("messagesReceived"));
-
-	const FName PacketsLost				  = FName(TEXT("packetsLost"));
-	const FName Jitter				      = FName(TEXT("jitter"));
-	const FName RoundTripTime			  = FName(TEXT("roundTripTime"));
-	const FName KeyFramesDecoded		  = FName(TEXT("keyFramesDecoded"));
-	const FName AudioLevel				  = FName(TEXT("audioLevel"));
-	const FName TotalSamplesDuration	  = FName(TEXT("totalSamplesDuration"));
-	const FName AvailableOutgoingBitrate  = FName(TEXT("availableOutgoingBitrate"));
-	const FName AvailableIncomingBitrate  = FName(TEXT("availableIncomingBitrate"));
-	const FName RetransmittedBytesReceived	  = FName(TEXT("retransmittedBytesReceived"));
-	const FName RetransmittedPacketsReceived  = FName(TEXT("retransmittedPacketsReceived"));
-
-	// Data-channel stats
-	const FName DataChannelBytesSent		  = FName(TEXT("data-channel-bytesSent"));
-	const FName DataChannelBytesReceived	  = FName(TEXT("data-channel-bytesReceived"));
-	const FName DataChannelMessagesSent		  = FName(TEXT("data-channel-messagesSent"));
-	const FName DataChannelMessagesReceived	  = FName(TEXT("data-channel-messagesReceived"));
-
-	// Pixel Streaming specific stats
-	const FName InputController			  = FName(TEXT("inputController"));
-
-} // namespace PixelStreaming2StatNames
\ No newline at end of file
+} // namespace PixelStreaming2StatNames
diff --git a/Engine/Plugins/Media/PixelStreaming2/Source/PixelStreaming2Editor/Public/IPixelStreaming2EditorModule.h b/Engine/Plugins/Media/PixelStreaming2/Source/PixelStreaming2Editor/Public/IPixelStreaming2EditorModule.h
--- a/Engine/Plugins/Media/PixelStreaming2/Source/PixelStreaming2Editor/Public/IPixelStreaming2EditorModule.h
+++ b/Engine/Plugins/Media/PixelStreaming2/Source/PixelStreaming2Editor/Public/IPixelStreaming2EditorModule.h
@@ -9,7 +9,8 @@
 #include "PixelStreaming2Servers.h"
 
 /**
- * The public interface of the Pixel Streaming Editor module.
+ * Public interface that manages Pixel Streaming specific functionality within the Unreal Editor.
+ * Provides various features for controlling Pixel Streaming, such as starting and stopping the streamer.
  */
 class PIXELSTREAMING2EDITOR_API IPixelStreaming2EditorModule : public IModuleInterface
 {
diff --git a/Engine/Plugins/Media/PixelStreaming2/Source/PixelStreaming2HMD/Public/IPixelStreaming2HMD.h b/Engine/Plugins/Media/PixelStreaming2/Source/PixelStreaming2HMD/Public/IPixelStreaming2HMD.h
--- a/Engine/Plugins/Media/PixelStreaming2/Source/PixelStreaming2HMD/Public/IPixelStreaming2HMD.h
+++ b/Engine/Plugins/Media/PixelStreaming2/Source/PixelStreaming2HMD/Public/IPixelStreaming2HMD.h
@@ -5,18 +5,26 @@
 #include "CoreMinimal.h"
 #include "Math/MathFwd.h"
 
+/**
+ * Interface class that allows for setting the transform data for a Head Mounted Display (HMD).
+ * The transform and projection matrix for each eye can be individually set.
+ */
 class PIXELSTREAMING2HMD_API IPixelStreaming2HMD
 {
 public:
 	/**
 	 * @brief Set the transform for the HMD.
-	 *
+	 * @param Transform The transform for the HMD to set.
 	 */
 	virtual void SetTransform(FTransform Transform) = 0;
 
 	/**
 	 * @brief Set the eye views, including the eye positions and their respective projection matrices
-	 *
+	 * @param Left The transform of the left eye relative to the HMD transform.
+	 * @param LeftProj The Projection matrix of the left eye.
+	 * @param Right The transform of the right eye relative to the HMD transform.
+	 * @param RightProj The Projection matrix of the right eye.
+	 * @param HMD The transform for the HMD to set.
 	 */
 	virtual void SetEyeViews(FTransform Left, FMatrix LeftProj, FTransform Right, FMatrix RightProj, FTransform HMD) = 0;
 };
diff --git a/Engine/Plugins/Media/PixelStreaming2/Source/PixelStreaming2HMD/Public/IPixelStreaming2HMDModule.h b/Engine/Plugins/Media/PixelStreaming2/Source/PixelStreaming2HMD/Public/IPixelStreaming2HMDModule.h
--- a/Engine/Plugins/Media/PixelStreaming2/Source/PixelStreaming2HMD/Public/IPixelStreaming2HMDModule.h
+++ b/Engine/Plugins/Media/PixelStreaming2/Source/PixelStreaming2HMD/Public/IPixelStreaming2HMDModule.h
@@ -9,7 +9,7 @@
 #include "PixelStreaming2HMDEnums.h"
 
 /**
- * The public interface of the Pixel Streaming HMD module.
+ * The public interface provides functionality for interacting with Pixel Streaming and Head Mounted Display (HMD) systems.
  */
 class PIXELSTREAMING2HMD_API IPixelStreaming2HMDModule : public IHeadMountedDisplayModule
 {
@@ -27,24 +27,25 @@ public:
 
 	/**
 	 * Checks to see if this module is loaded.
-	 *
 	 * @return True if the module is loaded.
 	 */
 	static inline bool IsAvailable() { return FModuleManager::Get().IsModuleLoaded("PixelStreaming2HMD"); }
 
 	/**
-	 * @brief Get the Pixel Streaming HMD object
-	 *
-	 * @return IPixelStreaming2HMD*
+	 * @brief Get the Pixel Streaming HMD object.
+	 * @return Pointer to the Pixel Streaming HMD.
 	 */
 	virtual IPixelStreaming2HMD* GetPixelStreaming2HMD() const = 0;
 
 	/**
-	 * @brief Get the Active XR System
-	 *
-	 * @return EPixelStreaming2XRSystem
+	 * @brief Get the Active XR System.
+	 * @return The Active XR System.
 	 */
 	virtual EPixelStreaming2XRSystem GetActiveXRSystem() = 0;
 
+	/**
+	 * @brief Set the Active XR System.
+	 * @param System The Active XR System to set.
+	 */
 	virtual void SetActiveXRSystem(EPixelStreaming2XRSystem System) = 0;
 };
diff --git a/Engine/Plugins/Media/PixelStreaming2/Source/PixelStreaming2HMD/Public/PixelStreaming2HMDEnums.h b/Engine/Plugins/Media/PixelStreaming2/Source/PixelStreaming2HMD/Public/PixelStreaming2HMDEnums.h
--- a/Engine/Plugins/Media/PixelStreaming2/Source/PixelStreaming2HMD/Public/PixelStreaming2HMDEnums.h
+++ b/Engine/Plugins/Media/PixelStreaming2/Source/PixelStreaming2HMD/Public/PixelStreaming2HMDEnums.h
@@ -4,9 +4,15 @@
 
 #include "HAL/Platform.h"
 
+/**
+ * Enums for different XR (Extended Reality) Systems that can be used with Pixel Streaming.
+ */
 enum class EPixelStreaming2XRSystem : uint8
 {
+	/** The XR system is unknown or not specified. */
 	Unknown,
+	/** HTC Vive XR system. */
 	HTCVive,
+	/** Oculus Quest XR system. */
 	Quest,
 };
\ No newline at end of file
diff --git a/Engine/Plugins/Media/PixelStreaming2/Source/PixelStreaming2Input/Public/IPixelStreaming2InputHandler.h b/Engine/Plugins/Media/PixelStreaming2/Source/PixelStreaming2Input/Public/IPixelStreaming2InputHandler.h
--- a/Engine/Plugins/Media/PixelStreaming2/Source/PixelStreaming2Input/Public/IPixelStreaming2InputHandler.h
+++ b/Engine/Plugins/Media/PixelStreaming2/Source/PixelStreaming2Input/Public/IPixelStreaming2InputHandler.h
@@ -98,6 +98,8 @@ public:
 	/**
 	 * @brief Register a custom function to execute when command JSON is received over the data channel: "{ type: "Command", YourCommand: YourCommandValue }".
 	 * Note: You can also override the default Pixel Streaming command handlers by setting handlers with the same name as those already used, e.g. "Stat.FPS".
+	 * @param CommandName The name of the command to handle. This corresponds to the key in the JSON message and is used to identify the command.
+	 * @param Handler The function that will be executed when the command is received.
 	 */
 	virtual void SetCommandHandler(const FString& CommandName, const CommandHandlerFn& Handler) = 0;
 
@@ -111,8 +113,7 @@ public:
 	/**
 	 * @brief Checks whether the given id has elevated priviledges.
 	 *
-	 * @return true The id is elevated.
-	 * @return false The id is not elevated.
+	 * @return True if id is elevated and false is not elevated.
 	 */
 	virtual bool IsElevated(const FString& Id) = 0;
 
@@ -135,7 +136,8 @@ public:
 	virtual TSharedPtr<IPixelStreaming2DataProtocol> GetFromStreamerProtocol() = 0;
 
 	/**
-	 * @brief Set the input handlers type. This controls whether input is routed to widgets or windows
+	 * @brief Set the input handlers type. This controls whether input is routed to widgets or windows.
+	 * @param InputType The input routing type.
 	 */
 	virtual void SetInputType(EPixelStreaming2InputType InputType) = 0;
 
diff --git a/Engine/Plugins/Media/PixelStreaming2/Source/PixelStreaming2Input/Public/IPixelStreaming2InputModule.h b/Engine/Plugins/Media/PixelStreaming2/Source/PixelStreaming2Input/Public/IPixelStreaming2InputModule.h
--- a/Engine/Plugins/Media/PixelStreaming2/Source/PixelStreaming2Input/Public/IPixelStreaming2InputModule.h
+++ b/Engine/Plugins/Media/PixelStreaming2/Source/PixelStreaming2Input/Public/IPixelStreaming2InputModule.h
@@ -7,7 +7,7 @@
 #include "IPixelStreaming2InputHandler.h"
 
 /**
- * The public interface of the Pixel Streaming Input module.
+ * Provides mechanism for managing input handlers and devices in Pixel Streaming.
  */
 class PIXELSTREAMING2INPUT_API IPixelStreaming2InputModule : public IInputDeviceModule
 {
@@ -16,7 +16,7 @@ public:
 	 * Singleton-like access to this module's interface.
 	 * Beware of calling this during the shutdown phase, though.  Your module might have been unloaded already.
 	 *
-	 * @return Returns singleton instance, loading the module on demand if needed
+	 * @return Returns singleton instance, loading the module on demand if needed.
 	 */
 	static inline IPixelStreaming2InputModule& Get()
 	{
@@ -31,16 +31,16 @@ public:
 	static inline bool IsAvailable() { return FModuleManager::Get().IsModuleLoaded("PixelStreaming2Input"); }
 
 	/**
-	 * @brief Create a Input Handler object
+	 * @brief Create a Input Handler object.
 	 *
-	 * @return TSharedPtr<IPixelStreaming2InputHandler> the input handler for this streamer
+	 * @return TSharedPtr<IPixelStreaming2InputHandler> the input handler for this streamer.
 	 */
 	virtual TSharedPtr<IPixelStreaming2InputHandler> CreateInputHandler() = 0;
 
 	/**
 	 * Attempts to create a new input device interface
-	 *
-	 * @return	Interface to the new input device, if we were able to successfully create one
+	 * @param InMessageHandler A reference to the message handler that will process input events from the device.
+	 * @return	Interface to the new input device, if we were able to successfully create one.
 	 */
 	virtual TSharedPtr<class IInputDevice> CreateInputDevice(const TSharedRef<FGenericApplicationMessageHandler>& InMessageHandler) override = 0;
 };
diff --git a/Engine/Plugins/Media/PixelStreaming2/Source/PixelStreaming2Servers/Public/PixelStreaming2Servers.h b/Engine/Plugins/Media/PixelStreaming2/Source/PixelStreaming2Servers/Public/PixelStreaming2Servers.h
--- a/Engine/Plugins/Media/PixelStreaming2/Source/PixelStreaming2Servers/Public/PixelStreaming2Servers.h
+++ b/Engine/Plugins/Media/PixelStreaming2/Source/PixelStreaming2Servers/Public/PixelStreaming2Servers.h
@@ -10,11 +10,9 @@
 #include "Misc/MonitoredProcess.h"
 
 /*
- * Utility namespace for launching Pixel Streaming servers signaling or selective forwarding>
- * This utility can launch embedded servers from:
- * 1. C++ classes in this module.
- * 2. From known server binaries.
- * 3. From the Github releases of the servers.
+ * Utility namespace for launching Pixel Streaming servers signaling or selective forwarding.
+ * This utility can launch embedded servers from C++ classes in this module, 
+ * from known server binaries and from the Github releases of the servers.
  */
 namespace UE::PixelStreaming2Servers
 {
