[PS2] Fix export on public API in Pixel Streaming 2.

Strings needed MODULE_API and definitions to be kept in cpp file.

diff --git a/Engine/Plugins/Media/PixelStreaming2/Source/PixelStreaming2/Private/PixelStreaming2StatNames.cpp b/Engine/Plugins/Media/PixelStreaming2/Source/PixelStreaming2/Private/PixelStreaming2StatNames.cpp
new file mode 100644
--- /dev/null
+++ b/Engine/Plugins/Media/PixelStreaming2/Source/PixelStreaming2/Private/PixelStreaming2StatNames.cpp
@@ -0,0 +1,77 @@
+// Copyright Epic Games, Inc. All Rights Reserved.
+#include "PixelStreaming2StatNames.h"
+
+namespace PixelStreaming2StatNames
+{
+	// Calculated stats
+	const FName Bitrate = FName(TEXT("bitrate"));
+	const FName BitrateMegabits = FName(TEXT("bitrateMbits"));
+	const FName TargetBitrateMegabits = FName(TEXT("targetBitrateMbits"));
+	const FName MeanSendDelay = FName(TEXT("captureToSend"));
+	const FName SourceFps = FName(TEXT("captureFps"));
+	const FName Fps = FName(TEXT("captureFps"));
+	const FName MeanEncodeTime = FName(TEXT("encodeTime"));
+	const FName EncodedFramesPerSecond = FName(TEXT("encodeFps"));
+	const FName DecodedFramesPerSecond = FName(TEXT("decodeFps"));
+	const FName MeanQPPerSecond = FName(TEXT("qp"));
+	const FName FramesSentPerSecond = FName(TEXT("transmitFps"));
+	const FName FramesReceivedPerSecond = FName(TEXT("receiveFps"));
+
+	// WebRTC stats
+	const FName JitterBufferDelay = FName(TEXT("jitterBufferDelay"));
+	const FName FramesSent = FName(TEXT("framesSent"));
+	const FName FramesReceived = FName(TEXT("framesReceived"));
+	const FName FramesPerSecond = FName(TEXT("framesPerSecond"));
+	const FName FramesDecoded = FName(TEXT("framesDecoded"));
+	const FName FramesDropped = FName(TEXT("framesDropped"));
+	const FName FramesCorrupted = FName(TEXT("framesCorrupted"));
+	const FName PartialFramesLost = FName(TEXT("partialFramesLost"));
+	const FName FullFramesLost = FName(TEXT("fullFramesLost"));
+	const FName HugeFramesSent = FName(TEXT("hugeFramesSent"));
+	const FName JitterBufferTargetDelay = FName(TEXT("jitterBufferTargetDelay"));
+	const FName InterruptionCount = FName(TEXT("interruptionCount"));
+	const FName TotalInterruptionDuration = FName(TEXT("totalInterruptionDuration"));
+	const FName FreezeCount = FName(TEXT("freezeCount"));
+	const FName PauseCount = FName(TEXT("pauseCount"));
+	const FName TotalFreezesDuration = FName(TEXT("totalFreezesDuration"));
+	const FName TotalPausesDuration = FName(TEXT("totalPausesDuration"));
+	const FName FirCount = FName(TEXT("firCount"));
+	const FName PliCount = FName(TEXT("pliCount"));
+	const FName NackCount = FName(TEXT("nackCount"));
+	const FName RetransmittedBytesSent = FName(TEXT("retransmittedBytesSent"));
+	const FName TargetBitrate = FName(TEXT("targetBitrate"));
+	const FName TotalEncodeBytesTarget = FName(TEXT("totalEncodedBytesTarget"));
+	const FName KeyFramesEncoded = FName(TEXT("keyFramesEncoded"));
+	const FName FrameWidth = FName(TEXT("frameWidth"));
+	const FName FrameHeight = FName(TEXT("frameHeight"));
+	const FName BytesSent = FName(TEXT("bytesSent"));
+	const FName BytesReceived = FName(TEXT("bytesReceived"));
+	const FName QPSum = FName(TEXT("qpSum"));
+	const FName TotalEncodeTime = FName(TEXT("totalEncodeTime"));
+	const FName TotalPacketSendDelay = FName(TEXT("totalPacketSendDelay"));
+	const FName FramesEncoded = FName(TEXT("framesEncoded"));
+	const FName AvgSendDelay = FName(TEXT("packetSendDelay"));
+	const FName MessagesSent = FName(TEXT("messagesSent"));
+	const FName MessagesReceived = FName(TEXT("messagesReceived"));
+
+	const FName PacketsLost = FName(TEXT("packetsLost"));
+	const FName Jitter = FName(TEXT("jitter"));
+	const FName RoundTripTime = FName(TEXT("roundTripTime"));
+	const FName KeyFramesDecoded = FName(TEXT("keyFramesDecoded"));
+	const FName AudioLevel = FName(TEXT("audioLevel"));
+	const FName TotalSamplesDuration = FName(TEXT("totalSamplesDuration"));
+	const FName AvailableOutgoingBitrate = FName(TEXT("availableOutgoingBitrate"));
+	const FName AvailableIncomingBitrate = FName(TEXT("availableIncomingBitrate"));
+	const FName RetransmittedBytesReceived = FName(TEXT("retransmittedBytesReceived"));
+	const FName RetransmittedPacketsReceived = FName(TEXT("retransmittedPacketsReceived"));
+
+	// Data-channel stats
+	const FName DataChannelBytesSent = FName(TEXT("data-channel-bytesSent"));
+	const FName DataChannelBytesReceived = FName(TEXT("data-channel-bytesReceived"));
+	const FName DataChannelMessagesSent = FName(TEXT("data-channel-messagesSent"));
+	const FName DataChannelMessagesReceived = FName(TEXT("data-channel-messagesReceived"));
+
+	// Pixel Streaming specific stats
+	const FName InputController = FName(TEXT("inputController"));
+
+} // namespace PixelStreaming2StatNames
diff --git a/Engine/Plugins/Media/PixelStreaming2/Source/PixelStreaming2Input/Private/PixelStreaming2InputEnums.cpp b/Engine/Plugins/Media/PixelStreaming2/Source/PixelStreaming2Input/Private/PixelStreaming2InputEnums.cpp
new file mode 100644
--- /dev/null
+++ b/Engine/Plugins/Media/PixelStreaming2/Source/PixelStreaming2Input/Private/PixelStreaming2InputEnums.cpp
@@ -0,0 +1,65 @@
+// Copyright Epic Games, Inc. All Rights Reserved.
+
+#include "PixelStreaming2InputEnums.h"
+
+namespace EPixelStreaming2ToStreamerMessage
+{
+	const FString IFrameRequest = FString(TEXT("IFrameRequest"));
+	const FString RequestQualityControl = FString(TEXT("RequestQualityControl"));
+	const FString FpsRequest = FString(TEXT("FpsRequest"));
+	const FString AverageBitrateRequest = FString(TEXT("AverageBitrateRequest"));
+	const FString StartStreaming = FString(TEXT("StartStreaming"));
+	const FString StopStreaming = FString(TEXT("StopStreaming"));
+	const FString LatencyTest = FString(TEXT("LatencyTest"));
+	const FString RequestInitialSettings = FString(TEXT("RequestInitialSettings"));
+	const FString TestEcho = FString(TEXT("TestEcho"));
+	const FString UIInteraction = FString(TEXT("UIInteraction"));
+	const FString Command = FString(TEXT("Command"));
+	const FString TextboxEntry = FString(TEXT("TextboxEntry"));
+	const FString KeyDown = FString(TEXT("KeyDown"));
+	const FString KeyUp = FString(TEXT("KeyUp"));
+	const FString KeyPress = FString(TEXT("KeyPress"));
+	const FString MouseEnter = FString(TEXT("MouseEnter"));
+	const FString MouseLeave = FString(TEXT("MouseLeave"));
+	const FString MouseDown = FString(TEXT("MouseDown"));
+	const FString MouseUp = FString(TEXT("MouseUp"));
+	const FString MouseMove = FString(TEXT("MouseMove"));
+	const FString MouseWheel = FString(TEXT("MouseWheel"));
+	const FString MouseDouble = FString(TEXT("MouseDouble"));
+	const FString TouchStart = FString(TEXT("TouchStart"));
+	const FString TouchEnd = FString(TEXT("TouchEnd"));
+	const FString TouchMove = FString(TEXT("TouchMove"));
+	const FString GamepadButtonPressed = FString(TEXT("GamepadButtonPressed"));
+	const FString GamepadButtonReleased = FString(TEXT("GamepadButtonReleased"));
+	const FString GamepadAnalog = FString(TEXT("GamepadAnalog"));
+	const FString GamepadConnected = FString(TEXT("GamepadConnected"));
+	const FString GamepadDisconnected = FString(TEXT("GamepadDisconnected"));
+	const FString XREyeViews = FString(TEXT("XREyeViews"));
+	const FString XRHMDTransform = FString(TEXT("XRHMDTransform"));
+	const FString XRControllerTransform = FString(TEXT("XRControllerTransform"));
+	const FString XRButtonPressed = FString(TEXT("XRButtonPressed"));
+	const FString XRButtonTouched = FString(TEXT("XRButtonTouched"));
+	const FString XRButtonReleased = FString(TEXT("XRButtonReleased"));
+	const FString XRAnalog = FString(TEXT("XRAnalog"));
+	const FString XRSystem = FString(TEXT("XRSystem"));
+	const FString XRButtonTouchReleased = FString(TEXT("XRButtonTouchReleased"));
+} // namespace EPixelStreaming2ToStreamerMessage
+
+namespace EPixelStreaming2FromStreamerMessage
+{
+	const FString QualityControlOwnership = FString(TEXT("QualityControlOwnership"));
+	const FString Response = FString(TEXT("Response"));
+	const FString Command = FString(TEXT("Command"));
+	const FString FreezeFrame = FString(TEXT("FreezeFrame"));
+	const FString UnfreezeFrame = FString(TEXT("UnfreezeFrame"));
+	const FString VideoEncoderAvgQP = FString(TEXT("VideoEncoderAvgQP"));
+	const FString LatencyTest = FString(TEXT("LatencyTest"));
+	const FString InitialSettings = FString(TEXT("InitialSettings"));
+	const FString FileExtension = FString(TEXT("FileExtension"));
+	const FString FileMimeType = FString(TEXT("FileMimeType"));
+	const FString FileContents = FString(TEXT("FileContents"));
+	const FString TestEcho = FString(TEXT("TestEcho"));
+	const FString InputControlOwnership = FString(TEXT("InputControlOwnership"));
+	const FString GamepadResponse = FString(TEXT("GamepadResponse"));
+	const FString Protocol = FString(TEXT("Protocol"));
+} // namespace EPixelStreaming2FromStreamerMessage
diff --git a/Engine/Plugins/Media/PixelStreaming2/Source/PixelStreaming2Input/Public/PixelStreaming2InputEnums.h b/Engine/Plugins/Media/PixelStreaming2/Source/PixelStreaming2Input/Public/PixelStreaming2InputEnums.h
--- a/Engine/Plugins/Media/PixelStreaming2/Source/PixelStreaming2Input/Public/PixelStreaming2InputEnums.h
+++ b/Engine/Plugins/Media/PixelStreaming2/Source/PixelStreaming2Input/Public/PixelStreaming2InputEnums.h
@@ -5,10 +5,9 @@
 #include "CoreMinimal.h"
 
 /**
- * @brief The data types supported when sending messages across the data channel to/from peers
- *
+ * @brief The data types supported when sending messages across the data channel to/from peers.
  */
-enum class PIXELSTREAMING2INPUT_API EPixelStreaming2MessageTypes
+enum class EPixelStreaming2MessageTypes
 {
 	Uint8 = 0,
 	Uint16 = 1,
@@ -21,19 +20,17 @@ enum class PIXELSTREAMING2INPUT_API EPixelStreaming2MessageTypes
 
 /**
  * @brief The message directions
- *
  */
-enum class PIXELSTREAMING2INPUT_API EPixelStreaming2MessageDirection : uint8
+enum class EPixelStreaming2MessageDirection : uint8
 {
 	ToStreamer = 0,
 	FromStreamer = 1
 };
 
 /**
- * @brief The possible actions pixel streaming supports. These actions help differentiate input received from the browser
- *
+ * @brief The possible actions pixel streaming supports. These actions help differentiate input received from the browser.
  */
-enum class PIXELSTREAMING2INPUT_API EPixelStreaming2InputAction : uint8
+enum class EPixelStreaming2InputAction : uint8
 {
 	X = 0,
 	Y = 1,
@@ -45,9 +42,8 @@ enum class PIXELSTREAMING2INPUT_API EPixelStreaming2InputAction : uint8
 
 /**
  * @brief The possible control schemes pixel streaming supports. RouteToWindow routes input at an application level. RouteToWidget routes input to a specific widget, ignoring the rest of the application.
- *
  */
-enum class PIXELSTREAMING2INPUT_API EPixelStreaming2InputType : uint8
+enum class EPixelStreaming2InputType : uint8
 {
 	RouteToWindow = 0,
 	RouteToWidget = 1
@@ -58,45 +54,45 @@ enum class PIXELSTREAMING2INPUT_API EPixelStreaming2InputType : uint8
  */
 namespace EPixelStreaming2ToStreamerMessage
 {
-	static FString IFrameRequest = FString(TEXT("IFrameRequest"));
-	static FString RequestQualityControl = FString(TEXT("RequestQualityControl"));
-	static FString FpsRequest = FString(TEXT("FpsRequest"));
-	static FString AverageBitrateRequest = FString(TEXT("AverageBitrateRequest"));
-	static FString StartStreaming = FString(TEXT("StartStreaming"));
-	static FString StopStreaming = FString(TEXT("StopStreaming"));
-	static FString LatencyTest = FString(TEXT("LatencyTest"));
-	static FString RequestInitialSettings = FString(TEXT("RequestInitialSettings"));
-	static FString TestEcho = FString(TEXT("TestEcho"));
-	static FString UIInteraction = FString(TEXT("UIInteraction"));
-	static FString Command = FString(TEXT("Command"));
-	static FString TextboxEntry = FString(TEXT("TextboxEntry"));
-	static FString KeyDown = FString(TEXT("KeyDown"));
-	static FString KeyUp = FString(TEXT("KeyUp"));
-	static FString KeyPress = FString(TEXT("KeyPress"));
-	static FString MouseEnter = FString(TEXT("MouseEnter"));
-	static FString MouseLeave = FString(TEXT("MouseLeave"));
-	static FString MouseDown = FString(TEXT("MouseDown"));
-	static FString MouseUp = FString(TEXT("MouseUp"));
-	static FString MouseMove = FString(TEXT("MouseMove"));
-	static FString MouseWheel = FString(TEXT("MouseWheel"));
-	static FString MouseDouble = FString(TEXT("MouseDouble"));
-	static FString TouchStart = FString(TEXT("TouchStart"));
-	static FString TouchEnd = FString(TEXT("TouchEnd"));
-	static FString TouchMove = FString(TEXT("TouchMove"));
-	static FString GamepadButtonPressed = FString(TEXT("GamepadButtonPressed"));
-	static FString GamepadButtonReleased = FString(TEXT("GamepadButtonReleased"));
-	static FString GamepadAnalog = FString(TEXT("GamepadAnalog"));
-	static FString GamepadConnected = FString(TEXT("GamepadConnected"));
-	static FString GamepadDisconnected = FString(TEXT("GamepadDisconnected"));
-	static FString XREyeViews = FString(TEXT("XREyeViews"));
-	static FString XRHMDTransform = FString(TEXT("XRHMDTransform"));
-	static FString XRControllerTransform = FString(TEXT("XRControllerTransform"));
-	static FString XRButtonPressed = FString(TEXT("XRButtonPressed"));
-	static FString XRButtonTouched = FString(TEXT("XRButtonTouched"));
-	static FString XRButtonReleased = FString(TEXT("XRButtonReleased"));
-	static FString XRAnalog = FString(TEXT("XRAnalog"));
-	static FString XRSystem = FString(TEXT("XRSystem"));
-	static FString XRButtonTouchReleased = FString(TEXT("XRButtonTouchReleased"));
+	PIXELSTREAMING2INPUT_API extern const FString IFrameRequest;
+	PIXELSTREAMING2INPUT_API extern const FString RequestQualityControl;
+	PIXELSTREAMING2INPUT_API extern const FString FpsRequest;
+	PIXELSTREAMING2INPUT_API extern const FString AverageBitrateRequest;
+	PIXELSTREAMING2INPUT_API extern const FString StartStreaming;
+	PIXELSTREAMING2INPUT_API extern const FString StopStreaming;
+	PIXELSTREAMING2INPUT_API extern const FString LatencyTest;
+	PIXELSTREAMING2INPUT_API extern const FString RequestInitialSettings;
+	PIXELSTREAMING2INPUT_API extern const FString TestEcho;
+	PIXELSTREAMING2INPUT_API extern const FString UIInteraction;
+	PIXELSTREAMING2INPUT_API extern const FString Command;
+	PIXELSTREAMING2INPUT_API extern const FString TextboxEntry;
+	PIXELSTREAMING2INPUT_API extern const FString KeyDown;
+	PIXELSTREAMING2INPUT_API extern const FString KeyUp;
+	PIXELSTREAMING2INPUT_API extern const FString KeyPress;
+	PIXELSTREAMING2INPUT_API extern const FString MouseEnter;
+	PIXELSTREAMING2INPUT_API extern const FString MouseLeave;
+	PIXELSTREAMING2INPUT_API extern const FString MouseDown;
+	PIXELSTREAMING2INPUT_API extern const FString MouseUp;
+	PIXELSTREAMING2INPUT_API extern const FString MouseMove;
+	PIXELSTREAMING2INPUT_API extern const FString MouseWheel;
+	PIXELSTREAMING2INPUT_API extern const FString MouseDouble;
+	PIXELSTREAMING2INPUT_API extern const FString TouchStart;
+	PIXELSTREAMING2INPUT_API extern const FString TouchEnd;
+	PIXELSTREAMING2INPUT_API extern const FString TouchMove;
+	PIXELSTREAMING2INPUT_API extern const FString GamepadButtonPressed;
+	PIXELSTREAMING2INPUT_API extern const FString GamepadButtonReleased;
+	PIXELSTREAMING2INPUT_API extern const FString GamepadAnalog;;
+	PIXELSTREAMING2INPUT_API extern const FString GamepadConnected;
+	PIXELSTREAMING2INPUT_API extern const FString GamepadDisconnected;
+	PIXELSTREAMING2INPUT_API extern const FString XREyeViews;
+	PIXELSTREAMING2INPUT_API extern const FString XRHMDTransform;
+	PIXELSTREAMING2INPUT_API extern const FString XRControllerTransform;
+	PIXELSTREAMING2INPUT_API extern const FString XRButtonPressed;
+	PIXELSTREAMING2INPUT_API extern const FString XRButtonTouched;
+	PIXELSTREAMING2INPUT_API extern const FString XRButtonReleased;
+	PIXELSTREAMING2INPUT_API extern const FString XRAnalog;
+	PIXELSTREAMING2INPUT_API extern const FString XRSystem;
+	PIXELSTREAMING2INPUT_API extern const FString XRButtonTouchReleased;
 } // namespace EPixelStreaming2ToStreamerMessage
 
 /**
@@ -104,19 +100,19 @@ namespace EPixelStreaming2ToStreamerMessage
  */
 namespace EPixelStreaming2FromStreamerMessage
 {
-	static FString QualityControlOwnership = FString(TEXT("QualityControlOwnership"));
-	static FString Response = FString(TEXT("Response"));
-	static FString Command = FString(TEXT("Command"));
-	static FString FreezeFrame = FString(TEXT("FreezeFrame"));
-	static FString UnfreezeFrame = FString(TEXT("UnfreezeFrame"));
-	static FString VideoEncoderAvgQP = FString(TEXT("VideoEncoderAvgQP"));
-	static FString LatencyTest = FString(TEXT("LatencyTest"));
-	static FString InitialSettings = FString(TEXT("InitialSettings"));
-	static FString FileExtension = FString(TEXT("FileExtension"));
-	static FString FileMimeType = FString(TEXT("FileMimeType"));
-	static FString FileContents = FString(TEXT("FileContents"));
-	static FString TestEcho = FString(TEXT("TestEcho"));
-	static FString InputControlOwnership = FString(TEXT("InputControlOwnership"));
-	static FString GamepadResponse = FString(TEXT("GamepadResponse"));
-	static FString Protocol = FString(TEXT("Protocol"));
+	PIXELSTREAMING2INPUT_API extern const FString QualityControlOwnership;
+	PIXELSTREAMING2INPUT_API extern const FString Response;
+	PIXELSTREAMING2INPUT_API extern const FString Command;
+	PIXELSTREAMING2INPUT_API extern const FString FreezeFrame;
+	PIXELSTREAMING2INPUT_API extern const FString UnfreezeFrame;
+	PIXELSTREAMING2INPUT_API extern const FString VideoEncoderAvgQP;
+	PIXELSTREAMING2INPUT_API extern const FString LatencyTest;
+	PIXELSTREAMING2INPUT_API extern const FString InitialSettings;
+	PIXELSTREAMING2INPUT_API extern const FString FileExtension;
+	PIXELSTREAMING2INPUT_API extern const FString FileMimeType;
+	PIXELSTREAMING2INPUT_API extern const FString FileContents;
+	PIXELSTREAMING2INPUT_API extern const FString TestEcho;
+	PIXELSTREAMING2INPUT_API extern const FString InputControlOwnership;
+	PIXELSTREAMING2INPUT_API extern const FString GamepadResponse;
+	PIXELSTREAMING2INPUT_API extern const FString Protocol;
 } // namespace EPixelStreaming2FromStreamerMessage
\ No newline at end of file
