[PS2] Fix colour distortion on decoded frames with VPX with Pixel Streaming 2.

Cause is because EpicRtc incorreectly handles frames with odd resolution (See RTCP-7246).

This is a temporary solution until EpicRtc is fixed. This crops decoded frames to a even resolution.

diff --git a/Engine/Plugins/Media/PixelStreaming2/Source/PixelStreaming2/Private/EpicRtcVideoDecoder.cpp b/Engine/Plugins/Media/PixelStreaming2/Source/PixelStreaming2/Private/EpicRtcVideoDecoder.cpp
--- a/Engine/Plugins/Media/PixelStreaming2/Source/PixelStreaming2/Private/EpicRtcVideoDecoder.cpp
+++ b/Engine/Plugins/Media/PixelStreaming2/Source/PixelStreaming2/Private/EpicRtcVideoDecoder.cpp
@@ -133,10 +133,15 @@ namespace UE::PixelStreaming2
 
 				if (DecodeResult.IsSuccess())
 				{
-					TSharedPtr<FPixelCaptureBufferI420> I420Buffer = MakeShared<FPixelCaptureBufferI420>(DecoderResource->GetWidth(), DecoderResource->GetHeight());
+					// TODO(Eden.Harris) RTCP-7247 EpicRtc currently has a bug where it incorrectly calculates stride for frames with odd resolution.
+					// To handle this, round width/height to a even number. When EpicRtc is fixed by RTCP-7246, this hack can be removed.
+					uint32 RoundedWidth = DecoderResource->GetWidth() & ~1;
+					uint32 RoundedHeight = DecoderResource->GetHeight() & ~1;
+					
+					TSharedPtr<FPixelCaptureBufferI420> I420Buffer = MakeShared<FPixelCaptureBufferI420>(RoundedWidth, RoundedHeight);
 
-					uint32 DataSizeY = I420Buffer->GetDataSizeY();
-					uint32 DataSizeUV = I420Buffer->GetDataSizeUV();
+					uint32 DataSizeY = DecoderResource->GetWidth() * DecoderResource->GetHeight();
+					uint32 DataSizeUV = ((DecoderResource->GetWidth() + 1) / 2) * ((DecoderResource->GetHeight() + 1) / 2);
 
 					CopyI420(
 						DecoderResource->GetRaw().Get(), DecoderResource->GetWidth(),
@@ -145,7 +150,7 @@ namespace UE::PixelStreaming2
 						I420Buffer->GetMutableDataY(), I420Buffer->GetStrideY(),
 						I420Buffer->GetMutableDataU(), I420Buffer->GetStrideUV(),
 						I420Buffer->GetMutableDataV(), I420Buffer->GetStrideUV(),
-						DecoderResource->GetWidth(), DecoderResource->GetHeight());
+						RoundedWidth, RoundedHeight);
 
 					DecodedFrame._buffer = new FEpicRtcVideoBufferI420(I420Buffer);
 				}
