[PS2] Add: Multiplexed DataTrack for use with the SFU
Adding a multiplex player data channels over one streamer data channel since EpicRtc doesn't support out of band channel negotiation with default signalling.

diff --git a/Engine/Plugins/Media/PixelStreaming2/Source/PixelStreaming2/Private/EpicRtcDataTrack.cpp b/Engine/Plugins/Media/PixelStreaming2/Source/PixelStreaming2/Private/EpicRtcDataTrack.cpp
--- a/Engine/Plugins/Media/PixelStreaming2/Source/PixelStreaming2/Private/EpicRtcDataTrack.cpp
+++ b/Engine/Plugins/Media/PixelStreaming2/Source/PixelStreaming2/Private/EpicRtcDataTrack.cpp
@@ -2,6 +2,7 @@
 
 #include "EpicRtcDataTrack.h"
 #include "Utils.h"
+#include "ToStringExtensions.h"
 
 namespace UE::PixelStreaming2
 {
@@ -18,19 +19,29 @@ namespace UE::PixelStreaming2
 	{
 	}
 
-	EpicRtcStringView FEpicRtcDataTrack::GetId() const
+	FEpicRtcDataTrack::FEpicRtcDataTrack(TSharedPtr<FEpicRtcDataTrack> InTrack, TWeakPtr<IPixelStreaming2DataProtocol> InDataProtocol)
+		: Track(InTrack->SendTrack ? InTrack->SendTrack : InTrack->Track)
+		, WeakDataProtocol(InDataProtocol)
 	{
-		return Track->GetId();
 	}
 
-	bool FEpicRtcDataTrack::SendArbitraryData(FString MessageType, const TArray64<uint8>& DataBytes) const
+	bool FEpicRtcDataTrack::IsActive() const
 	{
-		if (!Track || Track->GetState() != EpicRtcTrackState::Active)
+		if (!Track)
+		{
+			UE_LOG(LogPixelStreaming2, Error, TEXT("Cannot send message when datatrack is null."));
+			return false;
+		}
+		if (Track->GetState() != EpicRtcTrackState::Active)
 		{
-			UE_LOG(LogPixelStreaming2, Error, TEXT("Cannot send arbitrary data when datatrack is null."));
+			UE_LOG(LogPixelStreaming2, Error, TEXT("Cannot send message when datatrack is not active."));
 			return false;
 		}
+		return true;
+	}
 
+	bool FEpicRtcDataTrack::GetMessageId(const FString& MessageType, uint8& OutMessageId) const
+	{
 		TSharedPtr<IPixelStreaming2DataProtocol> DataProtocol = WeakDataProtocol.Pin();
 
 		if (!DataProtocol)
@@ -45,9 +56,42 @@ namespace UE::PixelStreaming2
 			UE_LOG(LogPixelStreaming2, Error, TEXT("Cannot send message called '%s' as it is not in the data protocol. Try GetTo/FromStreamerProtocol()->Add()"), *MessageType);
 			return false;
 		}
+		OutMessageId = Message->GetID();
+		return true;
+	}
+
+	bool FEpicRtcDataTrack::Send(TArray<uint8>& Buffer) const
+	{
+		EpicRtcDataFrameInput DataFrame{
+			._data = Buffer.GetData(),
+			._size = (uint32_t)Buffer.Num(),
+			._binary = true
+		};
+		TRefCountPtr<EpicRtcDataTrackInterface> OutgoingTrack = SendTrack ? SendTrack : Track;
+		EpicRtcBool SendResult = OutgoingTrack->PushFrame(DataFrame);
+		if (!SendResult)
+		{
+			UE_LOG(LogPixelStreaming2, Error, TEXT("DataTrack PushFrame return false"));
+		}
+		return static_cast<bool>(SendResult);
+	}
 
-		// The id of the message type we are about to send
-		const uint8 Type = Message->GetID();
+	EpicRtcStringView FEpicRtcDataTrack::GetId() const
+	{
+		return Track->GetId();
+	}
+
+	bool FEpicRtcDataTrack::SendArbitraryData(const FString& MessageType, const TArray64<uint8>& DataBytes) const
+	{
+		uint8 Type;
+		if (!IsActive())
+		{
+			return false;
+		}
+		if (!GetMessageId(MessageType, Type))
+		{
+			return false;
+		}
 
 		// int32 results in a maximum 4GB file (4,294,967,296 bytes)
 		const int32 DataSize = DataBytes.Num();
@@ -64,19 +108,17 @@ namespace UE::PixelStreaming2
 			int32 RemainingBytes = DataSize - BytesTransmitted;
 			int32 BytesToTransmit = FGenericPlatformMath::Min(MaxDataBytesPerMsg, RemainingBytes);
 
-			TArray<uint8> Buffer;
-			Buffer.SetNum(MessageHeader + BytesToTransmit);
-
-			size_t Pos = 0;
+			FBufferBuilder Builder(MessageHeader + BytesToTransmit);
+			PrependData(Builder);
 
 			// Write message type
-			Pos = SerializeToBuffer(Buffer, Pos, &Type, sizeof(Type));
+			Builder.Insert(Type);
 
 			// Write size of payload
-			Pos = SerializeToBuffer(Buffer, Pos, &DataSize, sizeof(DataSize));
+			Builder.Insert(DataSize);
 
 			// Write the data bytes payload
-			Pos = SerializeToBuffer(Buffer, Pos, DataBytes.GetData() + BytesTransmitted, BytesToTransmit);
+			Builder.Serialize(DataBytes.GetData() + BytesTransmitted, BytesToTransmit);
 
 			// TODO (Migration): RTCP-6489 We may need EpicRtc API surface to query the buffered amount in the datachannel so we don't flood it.
 			// uint64_t BufferBefore = SendChannel->buffered_amount();
@@ -87,17 +129,7 @@ namespace UE::PixelStreaming2
 			// 	BufferBefore = SendChannel->buffered_amount();
 			// }
 
-			EpicRtcDataFrameInput DataFrame{
-				._data = Buffer.GetData(),
-				._size = (uint32_t)Buffer.Num(),
-				._binary = true
-			};
-
-			EpicRtcBool SendResult = Track->PushFrame(DataFrame);
-			if (!SendResult)
-			{
-				UE_LOG(LogPixelStreaming2, Error, TEXT("DataTrack PushFrame return false"));
-			}
+			Send(Builder.Buffer);
 
 			// Increment the number of bytes transmitted
 			BytesTransmitted += BytesToTransmit;
@@ -106,4 +138,32 @@ namespace UE::PixelStreaming2
 		return true;
 	}
 
+	TSharedPtr<FEpicRtcMutliplexDataTrack> FEpicRtcMutliplexDataTrack::Create(TSharedPtr<FEpicRtcDataTrack> InTrack, TWeakPtr<IPixelStreaming2DataProtocol> InDataProtocol, const FString& InPlayerId)
+	{
+		TSharedPtr<FEpicRtcMutliplexDataTrack> DataTrack = TSharedPtr<FEpicRtcMutliplexDataTrack>(new FEpicRtcMutliplexDataTrack(InTrack, InDataProtocol, InPlayerId));
+		return DataTrack;
+	}
+
+	FEpicRtcMutliplexDataTrack::FEpicRtcMutliplexDataTrack(TSharedPtr<FEpicRtcDataTrack> InTrack, TWeakPtr<IPixelStreaming2DataProtocol> InDataProtocol, const FString& InPlayerId)
+		: FEpicRtcDataTrack(InTrack, InDataProtocol)
+		, PlayerId(InPlayerId)
+	{
+	}
+
+	void FEpicRtcMutliplexDataTrack::PrependData(FBufferBuilder& Builder) const
+	{
+		uint8 Type;
+		if (!GetMessageId(EPixelStreaming2FromStreamerMessage::Multiplexed, Type))
+		{
+			return;
+		}
+
+		uint16 StringLength = static_cast<uint16>(ValueSize(PlayerId));
+
+		Builder.Buffer.SetNum(Builder.Buffer.Num() + ValueSize(Type) + ValueSize(StringLength) + StringLength);
+
+		Builder.Insert(Type);
+		Builder.Insert(StringLength);
+		Builder.Insert<FString>(FString(PlayerId));
+	}
 } // namespace UE::PixelStreaming2
diff --git a/Engine/Plugins/Media/PixelStreaming2/Source/PixelStreaming2/Private/EpicRtcDataTrack.h b/Engine/Plugins/Media/PixelStreaming2/Source/PixelStreaming2/Private/EpicRtcDataTrack.h
--- a/Engine/Plugins/Media/PixelStreaming2/Source/PixelStreaming2/Private/EpicRtcDataTrack.h
+++ b/Engine/Plugins/Media/PixelStreaming2/Source/PixelStreaming2/Private/EpicRtcDataTrack.h
@@ -33,17 +33,24 @@ namespace UE::PixelStreaming2
 		return *Value;
 	}
 
-	struct BufferBuilder
+	struct FBufferBuilder
 	{
 		TArray<uint8> Buffer;
 		size_t		  Pos;
 
-		BufferBuilder(size_t size)
+		FBufferBuilder(size_t size)
 			: Pos(0)
 		{
 			Buffer.SetNum(size);
 		}
 
+		size_t Serialize(const void* Data, size_t DataSize)
+		{
+			check(Pos + DataSize <= Buffer.Num());
+			FMemory::Memcpy(Buffer.GetData() + Pos, Data, DataSize);
+			return Pos + DataSize;
+		}
+
 		template <typename T>
 		void Insert(T&& Value)
 		{
@@ -58,7 +65,9 @@ namespace UE::PixelStreaming2
 	class FEpicRtcDataTrack : public TSharedFromThis<FEpicRtcDataTrack>
 	{
 	public:
-		static TSharedPtr<FEpicRtcDataTrack> Create(TRefCountPtr<EpicRtcDataTrackInterface> InTrack, TWeakPtr<IPixelStreaming2DataProtocol> FromStreamerProtocol);
+		static TSharedPtr<FEpicRtcDataTrack> Create(TRefCountPtr<EpicRtcDataTrackInterface> InTrack, TWeakPtr<IPixelStreaming2DataProtocol> InDataProtocol);
+
+		virtual ~FEpicRtcDataTrack() = default;
 
 		/**
 		 * Sends a series of arguments to the data channel with the given type.
@@ -68,48 +77,20 @@ namespace UE::PixelStreaming2
 		template <typename... Args>
 		bool SendMessage(FString MessageType, Args... VarArgs) const
 		{
-			if (!Track)
-			{
-				UE_LOG(LogPixelStreaming2, Error, TEXT("Cannot send message when datatrack is null."));
-				return false;
-			}
-			if (Track->GetState() != EpicRtcTrackState::Active)
-			{
-				UE_LOG(LogPixelStreaming2, Error, TEXT("Cannot send message when is not active."));
-				return false;
-			}
-
-			TSharedPtr<IPixelStreaming2DataProtocol> DataProtocol = WeakDataProtocol.Pin();
-
-			if (!DataProtocol)
+			if (!IsActive())
 			{
-				UE_LOG(LogPixelStreaming2, Error, TEXT("Cannot send message, data protocol was null."));
 				return false;
 			}
 
-			TSharedPtr<IPixelStreaming2InputMessage> Message = DataProtocol->Find(MessageType);
-			if (!Message)
+			uint8 MessageId;
+			if (!GetMessageId(MessageType, MessageId))
 			{
-				UE_LOG(LogPixelStreaming2, Error, TEXT("Cannot send message called '%s' as it is not in the data protocol. Try GetTo/FromStreamerProtocol()->Add()"), *MessageType);
 				return false;
 			}
 
-			// Get the unique identifier of the message type we are about to send.
-			uint8 Type = Message->GetID();
-
-			BufferBuilder Builder(sizeof(Type) + (0 + ... + ValueSize(Forward<Args>(VarArgs))));
-			Builder.Insert(Type);
-			(Builder.Insert(Forward<Args>(VarArgs)), ...);
-
-			EpicRtcDataFrameInput InputFrame = {
-				._data = Builder.Buffer.GetData(),
-				._size = static_cast<uint32>(Builder.Buffer.Num()),
-				._binary = true
-			};
-
-			Track->PushFrame(InputFrame);
+			FBufferBuilder Builder = EncodeMessage(MessageId, Forward<Args>(VarArgs)...);
 
-			return true;
+			return Send(Builder.Buffer);
 		}
 
 		/**
@@ -118,7 +99,7 @@ namespace UE::PixelStreaming2
 		 * @param DataBytes The raw byte buffer to send.
 		 * @returns True of the message was successfully sent.
 		 */
-		bool SendArbitraryData(FString MessageType, const TArray64<uint8>& DataBytes) const;
+		bool SendArbitraryData(const FString& MessageType, const TArray64<uint8>& DataBytes) const;
 
 		/**
 		 * @return The id of the underlying EpicRtc data track.
@@ -130,11 +111,54 @@ namespace UE::PixelStreaming2
 		 */
 		EpicRtcTrackState GetState() const { return Track->GetState(); }
 
+		void SetSendTrack(TRefCountPtr<EpicRtcDataTrackInterface> InSendTrack) { SendTrack = InSendTrack; }
+
 	protected:
 		FEpicRtcDataTrack(TRefCountPtr<EpicRtcDataTrackInterface> InTrack, TWeakPtr<IPixelStreaming2DataProtocol> InDataProtocol);
+		FEpicRtcDataTrack(TSharedPtr<FEpicRtcDataTrack> InTrack, TWeakPtr<IPixelStreaming2DataProtocol> InDataProtocol);
+
+		virtual void PrependData(FBufferBuilder& Builder) const {};
+
+		bool IsActive() const;
+
+		bool GetMessageId(const FString& MessageType, uint8& OutMessageId) const;
 
 	private:
-		TRefCountPtr<EpicRtcDataTrackInterface>		 Track;
+		TRefCountPtr<EpicRtcDataTrackInterface> Track;
+
+		/**
+		 * Track that is used for sending data with Consumer/Producer architecture.
+		 */
+		TRefCountPtr<EpicRtcDataTrackInterface> SendTrack;
+
 		TWeakPtr<IPixelStreaming2DataProtocol> WeakDataProtocol;
+
+		bool Send(TArray<uint8>& Buffer) const;
+
+		template <typename... Args>
+		FBufferBuilder EncodeMessage(uint8 Type, Args... VarArgs) const
+		{
+
+			FBufferBuilder Builder(sizeof(Type) + (0 + ... + ValueSize(Forward<Args>(VarArgs))));
+			PrependData(Builder);
+			Builder.Insert(Type);
+			(Builder.Insert(Forward<Args>(VarArgs)), ...);
+
+			return MoveTemp(Builder);
+		}
+	};
+
+	class FEpicRtcMutliplexDataTrack : public FEpicRtcDataTrack
+	{
+	public:
+		static TSharedPtr<FEpicRtcMutliplexDataTrack> Create(TSharedPtr<FEpicRtcDataTrack> InTrack, TWeakPtr<IPixelStreaming2DataProtocol> InDataProtocol, const FString& InPlayerId);
+
+	protected:
+		FEpicRtcMutliplexDataTrack(TSharedPtr<FEpicRtcDataTrack> InTrack, TWeakPtr<IPixelStreaming2DataProtocol> InDataProtocol, const FString& InPlayerId);
+
+		virtual void PrependData(FBufferBuilder& Builder) const override;
+
+	private:
+		const FString PlayerId;
 	};
 } // namespace UE::PixelStreaming2
\ No newline at end of file
diff --git a/Engine/Plugins/Media/PixelStreaming2/Source/PixelStreaming2/Private/Stats.cpp b/Engine/Plugins/Media/PixelStreaming2/Source/PixelStreaming2/Private/Stats.cpp
--- a/Engine/Plugins/Media/PixelStreaming2/Source/PixelStreaming2/Private/Stats.cpp
+++ b/Engine/Plugins/Media/PixelStreaming2/Source/PixelStreaming2/Private/Stats.cpp
@@ -97,7 +97,7 @@ namespace UE::PixelStreaming2
 
 		PeerStats.Remove(PlayerId);
 
-		if (PlayerId == SFU_PLAYER_ID)
+		if (IsSFU(PlayerId))
 		{
 			TArray<FString> ToRemove;
 
diff --git a/Engine/Plugins/Media/PixelStreaming2/Source/PixelStreaming2/Private/Streamer.cpp b/Engine/Plugins/Media/PixelStreaming2/Source/PixelStreaming2/Private/Streamer.cpp
--- a/Engine/Plugins/Media/PixelStreaming2/Source/PixelStreaming2/Private/Streamer.cpp
+++ b/Engine/Plugins/Media/PixelStreaming2/Source/PixelStreaming2/Private/Streamer.cpp
@@ -313,7 +313,7 @@ namespace UE::PixelStreaming2
 	void FStreamer::SendAllPlayersMessage(FString MessageType, const FString& Descriptor)
 	{
 		Players->Apply([&MessageType, &Descriptor](FString PlayerId, FPlayerContext& PlayerContext) {
-			if (PlayerContext.DataTrack)
+			if (PlayerContext.DataTrack && !IsSFU(PlayerId))
 			{
 				PlayerContext.DataTrack->SendMessage(MessageType, Descriptor);
 			}
@@ -322,6 +322,10 @@ namespace UE::PixelStreaming2
 
 	void FStreamer::SendPlayerMessage(FString PlayerId, FString MessageType, const FString& Descriptor)
 	{
+		if (IsSFU(PlayerId))
+		{
+			return;
+		}
 		if (FPlayerContext* PlayerContext = Players->Find(PlayerId))
 		{
 			if (!PlayerContext->DataTrack)
@@ -496,6 +500,10 @@ namespace UE::PixelStreaming2
 
 	void FStreamer::ConsumeStats(FString PlayerId, FName StatName, float StatValue)
 	{
+		if (IsSFU(PlayerId))
+		{
+			return;
+		}
 
 		if (StatName != PixelStreaming2StatNames::MeanQPPerSecond)
 		{
@@ -543,11 +551,6 @@ namespace UE::PixelStreaming2
 
 		Players->Remove(PlayerId);
 
-		if (PlayerId == SFUPlayerId)
-		{
-			SFUPlayerId = INVALID_PLAYER_ID;
-		}
-
 		if (UPixelStreaming2Delegates* Delegates = UPixelStreaming2Delegates::Get())
 		{
 			Delegates->OnClosedConnection.Broadcast(StreamerId, PlayerId);
@@ -574,7 +577,6 @@ namespace UE::PixelStreaming2
 
 		VideoSourceGroup->RemoveAllVideoSources();
 		Players->Empty();
-		SFUPlayerId = INVALID_PLAYER_ID;
 		InputControllingId = INVALID_PLAYER_ID;
 		if (UPixelStreaming2Delegates* Delegates = UPixelStreaming2Delegates::Get())
 		{
@@ -587,7 +589,7 @@ namespace UE::PixelStreaming2
 	{
 		// Only time we automatically make a new peer the input controlling host is if they are the first peer (and not the SFU).
 		bool HostControlsInput = GetEnumFromCVar<EInputControllerMode>(UPixelStreaming2PluginSettings::CVarInputController) == EInputControllerMode::Host;
-		if (HostControlsInput && PlayerId != SFUPlayerId && InputControllingId == INVALID_PLAYER_ID)
+		if (HostControlsInput && !IsSFU(PlayerId) && InputControllingId == INVALID_PLAYER_ID)
 		{
 			InputControllingId = PlayerId;
 		}
@@ -621,7 +623,7 @@ namespace UE::PixelStreaming2
 					{
 						return false;
 					}
-					if (PlayerId == SFUPlayerId)
+					if (IsSFU(PlayerId))
 					{
 						return false;
 					}
@@ -771,6 +773,42 @@ namespace UE::PixelStreaming2
 		});
 	}
 
+	void FStreamer::HandleRelayStatusMessage(const uint8_t* Data, uint32_t Size, EpicRtcDataTrackInterface* DataTrack)
+	{
+		//skip type
+		Data++;
+		Size--;
+		FString PlayerId = ReadString(Data, Size);
+		checkf(Size > 0, TEXT("Malformed relay status message!"))
+		bool bIsOn = static_cast<bool>(Data[0]);
+
+		FString DataTrackId = ToString(DataTrack->GetId());
+		if (bIsOn)
+		{
+			UE_LOG(LogPixelStreaming2, Log, TEXT("FStreamer::HandleRelayStatusMessage() Adding new PlayerId [%s] with DataTrackId [%s]"), *PlayerId, *DataTrackId);
+			
+			FString SFUId;
+			if (FindPlayerByDataTrack(DataTrack, SFUId))
+			{
+				FPlayerContext* SFUContext = Players->Find(SFUId);
+				FPlayerContext& PlayerContext = Players->FindOrAdd(PlayerId);
+				PlayerContext.DataTrack = FEpicRtcMutliplexDataTrack::Create(SFUContext->DataTrack, InputHandler->GetFromStreamerProtocol(), PlayerId);
+				OnDataChannelOpen(PlayerId);
+			}
+			else
+			{
+				UE_LOG(LogPixelStreaming2, Log, TEXT("FStreamer::HandleRelayStatusMessage() Failed to find SFU PlayerContext"));
+			}
+		}
+		else
+		{
+			UE_LOG(LogPixelStreaming2, Log, TEXT("FStreamer::HandleRelayStatusMessage() Removing PlayerId [%s] with DataTrackId [%s]"), *PlayerId, *DataTrackId);
+
+			OnDataChannelClosed(PlayerId);
+			Players->Remove(PlayerId);
+		}
+	}
+
 	void FStreamer::TriggerMouseLeave(FString InStreamerId)
 	{
 		if (!IsEngineExitRequested() && StreamerId == InStreamerId)
@@ -1069,15 +1107,45 @@ namespace UE::PixelStreaming2
 			ParticipantConnection->AddVideoSource(VideoSource);
 		}
 
-		EpicRtcDataSource DataSource = {
-			._label = Participant->GetId(),
-			._maxRetransmitTime = 0,
-			._maxRetransmits = 0,
-			._isOrdered = true,
-			._protocol = EpicRtcDataSourceProtocol::Sctp
-		};
-
-		ParticipantConnection->AddDataSource(DataSource);
+		if (IsSFU(ParticipantId))
+		{
+			FString RecvLabel(TEXT("recv-datachannel"));
+			FUtf8String Utf8RecvLabel = *RecvLabel;
+			EpicRtcDataSource RecvDataSource = {
+				._label = ToEpicRtcStringView(Utf8RecvLabel),
+				._maxRetransmitTime = 0,
+				._maxRetransmits = 0,
+				._isOrdered = true,
+				._protocol = EpicRtcDataSourceProtocol::Sctp,
+				._negotiated = true,
+				._transportChannelId = 1
+			};
+			ParticipantConnection->AddDataSource(RecvDataSource);
+
+			FString SendLabel(TEXT("send-datachannel"));
+			FUtf8String Utf8SendLabel = *SendLabel;
+			EpicRtcDataSource SendDataSource = {
+				._label = ToEpicRtcStringView(Utf8RecvLabel),
+				._maxRetransmitTime = 0,
+				._maxRetransmits = 0,
+				._isOrdered = true,
+				._protocol = EpicRtcDataSourceProtocol::Sctp,
+				._negotiated = true,
+				._transportChannelId = 0
+			};
+			ParticipantConnection->AddDataSource(SendDataSource);
+		}
+		else
+		{
+			EpicRtcDataSource DataSource = {
+				._label = Participant->GetId(),
+				._maxRetransmitTime = 0,
+				._maxRetransmits = 0,
+				._isOrdered = true,
+				._protocol = EpicRtcDataSourceProtocol::Sctp
+			};
+			ParticipantConnection->AddDataSource(DataSource);
+		}
 
 		EpicRtcBitrate Bitrates = {
 			._minBitrateBps = UPixelStreaming2PluginSettings::CVarWebRTCMinBitrate.GetValueOnAnyThread(),
@@ -1158,7 +1226,14 @@ namespace UE::PixelStreaming2
 
 		if (FPlayerContext* PlayerContext = Players->Find(ParticipantId))
 		{
-			PlayerContext->DataTrack = FEpicRtcDataTrack::Create(DataTrack, InputHandler->GetFromStreamerProtocol());
+			if (!PlayerContext->DataTrack)
+			{
+				PlayerContext->DataTrack = FEpicRtcDataTrack::Create(DataTrack, InputHandler->GetFromStreamerProtocol());
+			}
+			else
+			{
+				PlayerContext->DataTrack->SetSendTrack(DataTrack);
+			}
 		}
 	}
 
@@ -1399,26 +1474,38 @@ namespace UE::PixelStreaming2
 
 	void FStreamer::OnDataTrackMessage(EpicRtcDataTrackInterface* DataTrack)
 	{
-		FString PlayerId = "";
-		bool	bFoundPlayer = FindPlayerByDataTrack(DataTrack, PlayerId);
 		FString DataTrackId = ToString(DataTrack->GetId());
-
-		if (!bFoundPlayer)
-		{
-			UE_LOG(LogPixelStreaming2, Log, TEXT("FStreamer::OnDataTrackMessage(Failed to find a player for data track [%s])"), *DataTrackId);
-			return;
-		}
-
 		TRefCountPtr<EpicRtcDataFrameInterface> DataFrame;
 		if (!DataTrack->PopFrame(DataFrame.GetInitReference()))
 		{
 			UE_LOG(LogPixelStreaming2, Log, TEXT("FStreamer::OnDataTrackMessage(Failed to PopFrame [%s])"), *DataTrackId);
 			return;
 		}
-
-		const uint8_t							 Type = DataFrame->Data()[0];
+		FString PlayerId;
+		const uint8_t* Data = DataFrame->Data();
+		uint32_t DataSize = DataFrame->Size();
+		uint8 Type = Data[0];
 		TSharedPtr<IPixelStreaming2DataProtocol> ToStreamerProtocol = InputHandler->GetToStreamerProtocol();
-
+		if (Type == ToStreamerProtocol->Find(EPixelStreaming2ToStreamerMessage::Multiplexed)->GetID())
+		{
+			//skip type
+			Data++;
+			DataSize--;
+			PlayerId = ReadString(Data, DataSize); 
+			Type = Data[0];
+			UE_LOG(LogPixelStreaming2, VeryVerbose, TEXT("FStreamer::OnDataTrackMessage(Received multiplexed message of type [%d] with PlayerId [%s])"), Type, *PlayerId);
+		}
+		else if (Type == ToStreamerProtocol->Find(EPixelStreaming2ToStreamerMessage::ChannelRelayStatus)->GetID())
+		{
+			HandleRelayStatusMessage(Data, DataSize, DataTrack);
+			return;
+		}
+		else if (!FindPlayerByDataTrack(DataTrack, PlayerId))
+		{
+			UE_LOG(LogPixelStreaming2, Log, TEXT("FStreamer::OnDataTrackMessage(Failed to find a player for data track [%s])"), *DataTrackId);
+			return;
+		}
+		
 		if (Type == ToStreamerProtocol->Find(EPixelStreaming2ToStreamerMessage::LatencyTest)->GetID())
 		{
 			SendLatencyReport(PlayerId);
@@ -1437,8 +1524,8 @@ namespace UE::PixelStreaming2
 			{
 				if (PlayerContext->DataTrack)
 				{
-					const size_t  DescriptorSize = (DataFrame->Size() - 1) / sizeof(TCHAR);
-					const TCHAR*  DescPtr = reinterpret_cast<const TCHAR*>(DataFrame->Data() + 1);
+					const size_t  DescriptorSize = (DataSize - 1) / sizeof(TCHAR);
+					const TCHAR*  DescPtr = reinterpret_cast<const TCHAR*>(Data + 1);
 					const FString Message(DescriptorSize, DescPtr);
 					PlayerContext->DataTrack->SendMessage(EPixelStreaming2FromStreamerMessage::TestEcho, Message);
 				}
@@ -1451,8 +1538,8 @@ namespace UE::PixelStreaming2
 			{
 				return;
 			}
-
-			TArray<uint8> MessageData((uint8*)DataFrame->Data(), (unsigned int)DataFrame->Size());
+			
+			TArray<uint8> MessageData(Data, DataSize);
 			if (InputHandler)
 			{
 				InputHandler->OnMessage(PlayerId, MessageData);
diff --git a/Engine/Plugins/Media/PixelStreaming2/Source/PixelStreaming2/Private/Streamer.h b/Engine/Plugins/Media/PixelStreaming2/Source/PixelStreaming2/Private/Streamer.h
--- a/Engine/Plugins/Media/PixelStreaming2/Source/PixelStreaming2/Private/Streamer.h
+++ b/Engine/Plugins/Media/PixelStreaming2/Source/PixelStreaming2/Private/Streamer.h
@@ -20,7 +20,6 @@ class IPixelStreaming2Module;
 namespace UE::PixelStreaming2
 {
 	static const FString INVALID_PLAYER_ID = FString(TEXT("Invalid Player Id"));
-	static const FString SFU_PLAYER_ID = FString(TEXT("1"));
 
 	class FStreamer : public IPixelStreaming2Streamer, public TSharedFromThis<FStreamer>
 	{
@@ -89,6 +88,7 @@ namespace UE::PixelStreaming2
 		void SendProtocol(FString PlayerId) const;
 		void SendPeerControllerMessages(FString PlayerId) const;
 		void SendLatencyReport(FString PlayerId) const;
+		void HandleRelayStatusMessage(const uint8_t* Data, uint32_t Size, EpicRtcDataTrackInterface* DataTrack);
 		bool FindPlayerByAudioTrack(EpicRtcAudioTrackInterface* AudioTrack, FString& OutPlayerId);
 		bool FindPlayerByVideoTrack(EpicRtcVideoTrackInterface* VideoTrack, FString& OutPlayerId);
 		bool FindPlayerByDataTrack(EpicRtcDataTrackInterface* DataTrack, FString& OutPlayerId);
@@ -105,7 +105,6 @@ namespace UE::PixelStreaming2
 
 		TSharedPtr<TThreadSafeMap<FString, FPlayerContext>> Players;
 
-		FString SFUPlayerId = INVALID_PLAYER_ID;
 		FString InputControllingId = INVALID_PLAYER_ID;
 
 		bool bSignallingConnected = false;
diff --git a/Engine/Plugins/Media/PixelStreaming2/Source/PixelStreaming2/Private/Utils.h b/Engine/Plugins/Media/PixelStreaming2/Source/PixelStreaming2/Private/Utils.h
--- a/Engine/Plugins/Media/PixelStreaming2/Source/PixelStreaming2/Private/Utils.h
+++ b/Engine/Plugins/Media/PixelStreaming2/Source/PixelStreaming2/Private/Utils.h
@@ -250,4 +250,29 @@ namespace UE::PixelStreaming2
 		return Pos + DataSize;
 	}
 
+	static const FString SFU_PLAYER_ID_PREFIX = FString(TEXT("SFU_"));
+
+	inline bool IsSFU(const FString& InPlayerId)
+	{
+		return InPlayerId.StartsWith(SFU_PLAYER_ID_PREFIX);
+	}
+
+	/**
+	 * Reads a string represented by 2 bytes length (in bytes) followed by UTF16 characters.
+	 * String and length are encoded in little endian format.
+	 */
+	inline FString ReadString(const uint8*& Data, uint32_t& Size)
+	{
+		uint16_t BytesLength = Data[1] << 8 | Data[0];
+		check(Size >= (uint32_t)(BytesLength + 2));
+		Data += 2;
+		Size -= 2;
+
+		FString Message(BytesLength / sizeof(TCHAR), reinterpret_cast<const TCHAR*>(Data));
+		Data += BytesLength;
+		Size -= BytesLength;
+
+		return Message;
+	}
+
 } // namespace UE::PixelStreaming2
diff --git a/Engine/Plugins/Media/PixelStreaming2/Source/PixelStreaming2Input/Private/DefaultDataProtocol.cpp b/Engine/Plugins/Media/PixelStreaming2/Source/PixelStreaming2Input/Private/DefaultDataProtocol.cpp
--- a/Engine/Plugins/Media/PixelStreaming2/Source/PixelStreaming2Input/Private/DefaultDataProtocol.cpp
+++ b/Engine/Plugins/Media/PixelStreaming2/Source/PixelStreaming2Input/Private/DefaultDataProtocol.cpp
@@ -130,6 +130,10 @@ namespace UE::PixelStreaming2Input
 		ToStreamerProtocol->AddInternal(EPixelStreaming2ToStreamerMessage::XRButtonTouchReleased, 117, {// Handedness, ButtonIdx, IsRepeat
 															EType::Uint8, EType::Uint8, EType::Uint8
 														});
+		
+		ToStreamerProtocol->AddInternal(EPixelStreaming2ToStreamerMessage::ChannelRelayStatus, 198); // id, 2 byte string length, string, uint8 flag
+		
+		ToStreamerProtocol->AddInternal(EPixelStreaming2ToStreamerMessage::Multiplexed, 199); // id, 2 byte string length, string, original message
 
 		// clang-format on
 		return ToStreamerProtocol;
@@ -154,6 +158,7 @@ namespace UE::PixelStreaming2Input
 		FromStreamerProtocol->AddInternal(EPixelStreaming2FromStreamerMessage::InputControlOwnership, 12);
 		FromStreamerProtocol->AddInternal(EPixelStreaming2FromStreamerMessage::GamepadResponse, 13);
 		FromStreamerProtocol->AddInternal(EPixelStreaming2FromStreamerMessage::Protocol, 255);
+		FromStreamerProtocol->AddInternal(EPixelStreaming2FromStreamerMessage::Multiplexed, 199);
 
 		return FromStreamerProtocol;
 	}
diff --git a/Engine/Plugins/Media/PixelStreaming2/Source/PixelStreaming2Input/Private/PixelStreaming2InputEnums.cpp b/Engine/Plugins/Media/PixelStreaming2/Source/PixelStreaming2Input/Private/PixelStreaming2InputEnums.cpp
--- a/Engine/Plugins/Media/PixelStreaming2/Source/PixelStreaming2Input/Private/PixelStreaming2InputEnums.cpp
+++ b/Engine/Plugins/Media/PixelStreaming2/Source/PixelStreaming2Input/Private/PixelStreaming2InputEnums.cpp
@@ -43,6 +43,8 @@ namespace EPixelStreaming2ToStreamerMessage
 	const FString XRAnalog = FString(TEXT("XRAnalog"));
 	const FString XRSystem = FString(TEXT("XRSystem"));
 	const FString XRButtonTouchReleased = FString(TEXT("XRButtonTouchReleased"));
+	const FString Multiplexed = FString(TEXT("Multiplexed"));
+	const FString ChannelRelayStatus = FString(TEXT("ChannelRelayStatus"));
 } // namespace EPixelStreaming2ToStreamerMessage
 
 namespace EPixelStreaming2FromStreamerMessage
@@ -61,5 +63,6 @@ namespace EPixelStreaming2FromStreamerMessage
 	const FString TestEcho = FString(TEXT("TestEcho"));
 	const FString InputControlOwnership = FString(TEXT("InputControlOwnership"));
 	const FString GamepadResponse = FString(TEXT("GamepadResponse"));
+	const FString Multiplexed = FString(TEXT("Multiplexed"));
 	const FString Protocol = FString(TEXT("Protocol"));
 } // namespace EPixelStreaming2FromStreamerMessage
diff --git a/Engine/Plugins/Media/PixelStreaming2/Source/PixelStreaming2Input/Public/PixelStreaming2InputEnums.h b/Engine/Plugins/Media/PixelStreaming2/Source/PixelStreaming2Input/Public/PixelStreaming2InputEnums.h
--- a/Engine/Plugins/Media/PixelStreaming2/Source/PixelStreaming2Input/Public/PixelStreaming2InputEnums.h
+++ b/Engine/Plugins/Media/PixelStreaming2/Source/PixelStreaming2Input/Public/PixelStreaming2InputEnums.h
@@ -93,6 +93,8 @@ namespace EPixelStreaming2ToStreamerMessage
 	PIXELSTREAMING2INPUT_API extern const FString XRAnalog;
 	PIXELSTREAMING2INPUT_API extern const FString XRSystem;
 	PIXELSTREAMING2INPUT_API extern const FString XRButtonTouchReleased;
+	PIXELSTREAMING2INPUT_API extern const FString Multiplexed;
+	PIXELSTREAMING2INPUT_API extern const FString ChannelRelayStatus;
 } // namespace EPixelStreaming2ToStreamerMessage
 
 /**
@@ -114,5 +116,6 @@ namespace EPixelStreaming2FromStreamerMessage
 	PIXELSTREAMING2INPUT_API extern const FString TestEcho;
 	PIXELSTREAMING2INPUT_API extern const FString InputControlOwnership;
 	PIXELSTREAMING2INPUT_API extern const FString GamepadResponse;
+	PIXELSTREAMING2INPUT_API extern const FString Multiplexed;
 	PIXELSTREAMING2INPUT_API extern const FString Protocol;
 } // namespace EPixelStreaming2FromStreamerMessage
\ No newline at end of file
