[PS2] QOL: Tidy up code to be more consistent throughout the plugin

diff --git a/Engine/Plugins/Media/PixelStreaming2/Source/PixelStreaming2/Internal/Blueprints/PixelStreaming2MediaTexture.h b/Engine/Plugins/Media/PixelStreaming2/Source/PixelStreaming2/Internal/Blueprints/PixelStreaming2MediaTexture.h
--- a/Engine/Plugins/Media/PixelStreaming2/Source/PixelStreaming2/Internal/Blueprints/PixelStreaming2MediaTexture.h
+++ b/Engine/Plugins/Media/PixelStreaming2/Source/PixelStreaming2/Internal/Blueprints/PixelStreaming2MediaTexture.h
@@ -4,7 +4,6 @@
 
 #include "Engine/Texture2DDynamic.h"
 #include "IPixelStreaming2VideoConsumer.h"
-#include "RenderTargetPool.h"
 
 #include "PixelStreaming2MediaTexture.generated.h"
 
diff --git a/Engine/Plugins/Media/PixelStreaming2/Source/PixelStreaming2/Internal/Blueprints/PixelStreaming2VideoProducer.h b/Engine/Plugins/Media/PixelStreaming2/Source/PixelStreaming2/Internal/Blueprints/PixelStreaming2VideoProducer.h
--- a/Engine/Plugins/Media/PixelStreaming2/Source/PixelStreaming2/Internal/Blueprints/PixelStreaming2VideoProducer.h
+++ b/Engine/Plugins/Media/PixelStreaming2/Source/PixelStreaming2/Internal/Blueprints/PixelStreaming2VideoProducer.h
@@ -2,13 +2,12 @@
 
 #pragma once
 
-#include "Engine/Texture2DDynamic.h"
 #include "Engine/TextureRenderTarget2D.h"
 #include "IPixelStreaming2VideoProducer.h"
 
 #include "PixelStreaming2VideoProducer.generated.h"
 
-/*
+/**
  * The base video producer blueprint class
  */
 UCLASS(NotBlueprintType, NotBlueprintable, Category = "PixelStreaming2", META = (DisplayName = "Video Producer"))
@@ -23,7 +22,7 @@ protected:
 	TSharedPtr<IPixelStreaming2VideoProducer> VideoProducer;
 };
 
-/*
+/**
  * A video producer that streams the Unreal Engine's back buffer
  */
 UCLASS(NotBlueprintType, NotBlueprintable, Category = "PixelStreaming2", META = (DisplayName = "Back Buffer Video Producer"))
@@ -35,7 +34,7 @@ public:
 	virtual TSharedPtr<IPixelStreaming2VideoProducer> GetVideoProducer() override;
 };
 
-/*
+/**
  * A video producer that streams the Unreal Engine's back buffer but uses the MediaIO Framework to capture the frame rather than Pixel Capture
  */
 UCLASS(NotBlueprintType, NotBlueprintable, Category = "PixelStreaming2", META = (DisplayName = "Media Capture Video Producer"))
@@ -47,8 +46,8 @@ public:
 	virtual TSharedPtr<IPixelStreaming2VideoProducer> GetVideoProducer() override;
 };
 
-/*
- * A video producer that streams the specified render target 
+/**
+ * A video producer that streams the specified render target
  */
 UCLASS(NotBlueprintType, NotBlueprintable, Category = "PixelStreaming2", META = (DisplayName = "Render Target Video Producer"))
 class PIXELSTREAMING2_API UPixelStreaming2VideoProducerRenderTarget : public UPixelStreaming2VideoProducer
diff --git a/Engine/Plugins/Media/PixelStreaming2/Source/PixelStreaming2/Internal/EpicRtcAllocator.h b/Engine/Plugins/Media/PixelStreaming2/Source/PixelStreaming2/Internal/EpicRtcAllocator.h
--- a/Engine/Plugins/Media/PixelStreaming2/Source/PixelStreaming2/Internal/EpicRtcAllocator.h
+++ b/Engine/Plugins/Media/PixelStreaming2/Source/PixelStreaming2/Internal/EpicRtcAllocator.h
@@ -2,16 +2,17 @@
 
 #pragma once
 
-#include "epic_rtc/common/memory.h"
 #include "HAL/UnrealMemory.h"
 
+#include "epic_rtc/common/memory.h"
+
 namespace UE::PixelStreaming2
 {
 	class PIXELSTREAMING2_API FEpicRtcAllocator : public EpicRtcMemoryInterface
 	{
 	public:
 		FEpicRtcAllocator() = default;
-		~FEpicRtcAllocator() = default;
+		virtual ~FEpicRtcAllocator() = default;
 
 		[[nodiscard]] virtual void* Allocate(uint64_t Size, uint64_t Alignment, const char* Tag) override
 		{
diff --git a/Engine/Plugins/Media/PixelStreaming2/Source/PixelStreaming2/Internal/EpicRtcAudioTrackObserver.h b/Engine/Plugins/Media/PixelStreaming2/Source/PixelStreaming2/Internal/EpicRtcAudioTrackObserver.h
--- a/Engine/Plugins/Media/PixelStreaming2/Source/PixelStreaming2/Internal/EpicRtcAudioTrackObserver.h
+++ b/Engine/Plugins/Media/PixelStreaming2/Source/PixelStreaming2/Internal/EpicRtcAudioTrackObserver.h
@@ -2,7 +2,7 @@
 
 #pragma once
 
-#include "EpicRtcMemory.h"
+#include "Templates/RefCounting.h"
 #include "Templates/SharedPointer.h"
 
 #include "epic_rtc/core/audio/audio_track_observer.h"
@@ -11,28 +11,28 @@ namespace UE::PixelStreaming2
 {
 	class FEpicRtcManager;
 
-	class PIXELSTREAMING2_API FEpicRtcAudioTrackObserver : public EpicRtcAudioTrackObserverInterface, public TEpicRtcRefCountPtr<FEpicRtcAudioTrackObserver>
+	class PIXELSTREAMING2_API FEpicRtcAudioTrackObserver : public EpicRtcAudioTrackObserverInterface, public TRefCountingMixin<FEpicRtcAudioTrackObserver>
 	{
 	public:
 		FEpicRtcAudioTrackObserver(TWeakPtr<FEpicRtcManager> Manager);
+		virtual ~FEpicRtcAudioTrackObserver() = default;
 
 	private:
-		/* Begin EpicRtcAudioTrackObserverInterface */
+		// Begin EpicRtcAudioTrackObserverInterface
 		virtual void OnAudioTrackMuted(EpicRtcAudioTrackInterface* AudioTrack, EpicRtcBool bIsMuted) override;
 		virtual void OnAudioTrackFrame(EpicRtcAudioTrackInterface* AudioTrack, const EpicRtcAudioFrame& Frame) override;
 		virtual void OnAudioTrackRemoved(EpicRtcAudioTrackInterface* AudioTrack) override;
 		virtual void OnAudioTrackState(EpicRtcAudioTrackInterface* AudioTrack, const EpicRtcTrackState State) override;
-		/* End EpicRtcAudioTrackObserverInterface */
+		// End EpicRtcAudioTrackObserverInterface
 
 	public:
-		/* Begin EpicRtcRefCountInterface */
-		virtual uint32_t AddRef() override final { return TEpicRtcRefCountPtr<FEpicRtcAudioTrackObserver>::AddRef(); }
-		virtual uint32_t Release() override final { return TEpicRtcRefCountPtr<FEpicRtcAudioTrackObserver>::Release(); }
-		virtual uint32_t Count() const override final { return TEpicRtcRefCountPtr<FEpicRtcAudioTrackObserver>::GetRefCount(); }
-		/* End EpicRtcRefCountInterface */
+		// Begin EpicRtcRefCountInterface
+		virtual uint32_t AddRef() override final { return TRefCountingMixin<FEpicRtcAudioTrackObserver>::AddRef(); }
+		virtual uint32_t Release() override final { return TRefCountingMixin<FEpicRtcAudioTrackObserver>::Release(); }
+		virtual uint32_t Count() const override final { return TRefCountingMixin<FEpicRtcAudioTrackObserver>::GetRefCount(); }
+		// End EpicRtcRefCountInterface
 
 	private:
 		TWeakPtr<FEpicRtcManager> Manager;
 	};
-
 } // namespace UE::PixelStreaming2
\ No newline at end of file
diff --git a/Engine/Plugins/Media/PixelStreaming2/Source/PixelStreaming2/Internal/EpicRtcAudioTrackObserverFactory.h b/Engine/Plugins/Media/PixelStreaming2/Source/PixelStreaming2/Internal/EpicRtcAudioTrackObserverFactory.h
--- a/Engine/Plugins/Media/PixelStreaming2/Source/PixelStreaming2/Internal/EpicRtcAudioTrackObserverFactory.h
+++ b/Engine/Plugins/Media/PixelStreaming2/Source/PixelStreaming2/Internal/EpicRtcAudioTrackObserverFactory.h
@@ -3,30 +3,32 @@
 #pragma once
 
 #include "EpicRtcAudioTrackObserver.h"
+#include "Templates/RefCounting.h"
+#include "Templates/SharedPointer.h"
 
 #include "epic_rtc/core/audio/audio_track_observer.h"
 
 namespace UE::PixelStreaming2
 {
-
 	class FEpicRtcManager;
 
 	class PIXELSTREAMING2_API FEpicRtcAudioTrackObserverFactory : public EpicRtcAudioTrackObserverFactoryInterface, public TRefCountingMixin<FEpicRtcAudioTrackObserverFactory>
 	{
 	public:
 		FEpicRtcAudioTrackObserverFactory(TWeakPtr<FEpicRtcManager> Manager);
+		virtual ~FEpicRtcAudioTrackObserverFactory() = default;
 
 	public:
-		/* Begin EpicRtcAudioTrackObserverFactoryInterface */
+		// Begin EpicRtcAudioTrackObserverFactoryInterface 
 		virtual EpicRtcErrorCode CreateAudioTrackObserver(const EpicRtcStringView ParticipantId, const EpicRtcStringView AudioTrackId, EpicRtcAudioTrackObserverInterface** OutAudioTrackObserver) override;
-		/* End EpicRtcAudioTrackObserverFactoryInterface */
+		// End EpicRtcAudioTrackObserverFactoryInterface 
 
 	public:
-		/* Begin EpicRtcRefCountInterface */
+		// Begin EpicRtcRefCountInterface 
 		virtual uint32_t AddRef() override final { return TRefCountingMixin<FEpicRtcAudioTrackObserverFactory>::AddRef(); }
 		virtual uint32_t Release() override final { return TRefCountingMixin<FEpicRtcAudioTrackObserverFactory>::Release(); }
 		virtual uint32_t Count() const override final { return TRefCountingMixin<FEpicRtcAudioTrackObserverFactory>::GetRefCount(); }
-		/* End EpicRtcRefCountInterface */
+		// End EpicRtcRefCountInterface 
 
 	private:
 		TWeakPtr<FEpicRtcManager> Manager;
diff --git a/Engine/Plugins/Media/PixelStreaming2/Source/PixelStreaming2/Internal/EpicRtcConferenceUtils.h b/Engine/Plugins/Media/PixelStreaming2/Source/PixelStreaming2/Internal/EpicRtcConferenceUtils.h
--- a/Engine/Plugins/Media/PixelStreaming2/Source/PixelStreaming2/Internal/EpicRtcConferenceUtils.h
+++ b/Engine/Plugins/Media/PixelStreaming2/Source/PixelStreaming2/Internal/EpicRtcConferenceUtils.h
@@ -10,7 +10,6 @@
 
 namespace UE::PixelStreaming2
 {
-
 	class PIXELSTREAMING2_API FEpicRtcTickConferenceTask : public FEpicRtcTickableTask
 	{
 	public:
@@ -34,7 +33,7 @@ namespace UE::PixelStreaming2
 			}
 		}
 
-		/** Begin FEpicRtcTickableTask */
+		// Begin FEpicRtcTickableTask
 		virtual void Tick(float DeltaMs) override
 		{
 			if (EpicRtcConference)
@@ -74,7 +73,7 @@ namespace UE::PixelStreaming2
 		{
 			return TaskName;
 		}
-		/** End FEpicRtcTickableTask */
+		// End FEpicRtcTickableTask
 
 	private:
 		TRefCountPtr<EpicRtcConferenceInterface> EpicRtcConference;
diff --git a/Engine/Plugins/Media/PixelStreaming2/Source/PixelStreaming2/Internal/EpicRtcDataTrackObserver.h b/Engine/Plugins/Media/PixelStreaming2/Source/PixelStreaming2/Internal/EpicRtcDataTrackObserver.h
--- a/Engine/Plugins/Media/PixelStreaming2/Source/PixelStreaming2/Internal/EpicRtcDataTrackObserver.h
+++ b/Engine/Plugins/Media/PixelStreaming2/Source/PixelStreaming2/Internal/EpicRtcDataTrackObserver.h
@@ -2,34 +2,34 @@
 
 #pragma once
 
-#include "EpicRtcMemory.h"
+#include "Templates/RefCounting.h"
 #include "Templates/SharedPointer.h"
 
 #include "epic_rtc/core/data_track_observer.h"
 
 namespace UE::PixelStreaming2
 {
-
 	class FEpicRtcManager;
 
-	class PIXELSTREAMING2_API FEpicRtcDataTrackObserver : public EpicRtcDataTrackObserverInterface, public TEpicRtcRefCountPtr<FEpicRtcDataTrackObserver>
+	class PIXELSTREAMING2_API FEpicRtcDataTrackObserver : public EpicRtcDataTrackObserverInterface, public TRefCountingMixin<FEpicRtcDataTrackObserver>
 	{
 	public:
 		FEpicRtcDataTrackObserver(TWeakPtr<FEpicRtcManager> Manager);
+		virtual ~FEpicRtcDataTrackObserver() = default;
 
 	private:
-		/* Begin EpicRtcDataTrackObserverInterface */
+		// Begin EpicRtcDataTrackObserverInterface
 		virtual void OnDataTrackState(EpicRtcDataTrackInterface* DataTrack, const EpicRtcTrackState State) override;
 		virtual void OnDataTrackMessage(EpicRtcDataTrackInterface* DataTrack) override;
 		virtual void OnDataTrackError(EpicRtcDataTrackInterface*, const EpicRtcErrorCode) override {};
-		/* End EpicRtcDataTrackObserverInterface */
+		// End EpicRtcDataTrackObserverInterface
 
 	public:
-		/* Begin EpicRtcRefCountInterface */
-		virtual uint32_t AddRef() override final { return TEpicRtcRefCountPtr<FEpicRtcDataTrackObserver>::AddRef(); }
-		virtual uint32_t Release() override final { return TEpicRtcRefCountPtr<FEpicRtcDataTrackObserver>::Release(); }
-		virtual uint32_t Count() const override final { return TEpicRtcRefCountPtr<FEpicRtcDataTrackObserver>::GetRefCount(); }
-		/* End EpicRtcRefCountInterface */
+		// Begin EpicRtcRefCountInterface
+		virtual uint32_t AddRef() override final { return TRefCountingMixin<FEpicRtcDataTrackObserver>::AddRef(); }
+		virtual uint32_t Release() override final { return TRefCountingMixin<FEpicRtcDataTrackObserver>::Release(); }
+		virtual uint32_t Count() const override final { return TRefCountingMixin<FEpicRtcDataTrackObserver>::GetRefCount(); }
+		// End EpicRtcRefCountInterface
 
 	private:
 		TWeakPtr<FEpicRtcManager> Manager;
diff --git a/Engine/Plugins/Media/PixelStreaming2/Source/PixelStreaming2/Internal/EpicRtcDataTrackObserverFactory.h b/Engine/Plugins/Media/PixelStreaming2/Source/PixelStreaming2/Internal/EpicRtcDataTrackObserverFactory.h
--- a/Engine/Plugins/Media/PixelStreaming2/Source/PixelStreaming2/Internal/EpicRtcDataTrackObserverFactory.h
+++ b/Engine/Plugins/Media/PixelStreaming2/Source/PixelStreaming2/Internal/EpicRtcDataTrackObserverFactory.h
@@ -10,25 +10,25 @@
 
 namespace UE::PixelStreaming2
 {
-
 	class FEpicRtcManager;
 
 	class PIXELSTREAMING2_API FEpicRtcDataTrackObserverFactory : public EpicRtcDataTrackObserverFactoryInterface, public TRefCountingMixin<FEpicRtcDataTrackObserverFactory>
 	{
 	public:
 		FEpicRtcDataTrackObserverFactory(TWeakPtr<FEpicRtcManager> Manager);
+		virtual ~FEpicRtcDataTrackObserverFactory() = default;
 
 	public:
-		/* Begin EpicRtcDataTrackObserverFactoryInterface */
+		// Begin EpicRtcDataTrackObserverFactoryInterface 
 		virtual EpicRtcErrorCode CreateDataTrackObserver(const EpicRtcStringView ParticipantId, const EpicRtcStringView DataTrackId, EpicRtcDataTrackObserverInterface** OutDataTrackObserver) override;
-		/* End EpicRtcDataTrackObserverFactoryInterface */
+		// End EpicRtcDataTrackObserverFactoryInterface 
 
 	public:
-		/* Begin EpicRtcRefCountInterface */
+		// Begin EpicRtcRefCountInterface 
 		virtual uint32_t AddRef() override final { return TRefCountingMixin<FEpicRtcDataTrackObserverFactory>::AddRef(); }
 		virtual uint32_t Release() override final { return TRefCountingMixin<FEpicRtcDataTrackObserverFactory>::Release(); }
 		virtual uint32_t Count() const override final { return TRefCountingMixin<FEpicRtcDataTrackObserverFactory>::GetRefCount(); }
-		/* End EpicRtcRefCountInterface */
+		// End EpicRtcRefCountInterface 
 
 	private:
 		TWeakPtr<FEpicRtcManager> Manager;
diff --git a/Engine/Plugins/Media/PixelStreaming2/Source/PixelStreaming2/Internal/EpicRtcLogging.h b/Engine/Plugins/Media/PixelStreaming2/Source/PixelStreaming2/Internal/EpicRtcLogging.h
--- a/Engine/Plugins/Media/PixelStreaming2/Source/PixelStreaming2/Internal/EpicRtcLogging.h
+++ b/Engine/Plugins/Media/PixelStreaming2/Source/PixelStreaming2/Internal/EpicRtcLogging.h
@@ -3,7 +3,6 @@
 #pragma once
 
 #include "Logging/LogMacros.h"
-#include "Misc/AssertionMacros.h"
 
 #include "epic_rtc/common/logging.h"
 
@@ -16,4 +15,4 @@ namespace UE::PixelStreaming2
 	public:
 		virtual void Log(const EpicRtcLogMessage& Message) override;
 	};
-}
\ No newline at end of file
+} // namespace UE::PixelStreaming2
\ No newline at end of file
diff --git a/Engine/Plugins/Media/PixelStreaming2/Source/PixelStreaming2/Internal/EpicRtcManager.h b/Engine/Plugins/Media/PixelStreaming2/Source/PixelStreaming2/Internal/EpicRtcManager.h
--- a/Engine/Plugins/Media/PixelStreaming2/Source/PixelStreaming2/Internal/EpicRtcManager.h
+++ b/Engine/Plugins/Media/PixelStreaming2/Source/PixelStreaming2/Internal/EpicRtcManager.h
@@ -16,7 +16,6 @@
 
 namespace UE::PixelStreaming2
 {
-
 	/**
 	 * A base class for managing epic rtc sessions, rooms and participants
 	 *
@@ -28,7 +27,7 @@ namespace UE::PixelStreaming2
 		virtual ~FEpicRtcManager() = default;
 
 	public:
-		/* Begin FEpicRtcSessionObserver Callbacks */
+		// Begin FEpicRtcSessionObserver Callbacks
 		DECLARE_TS_MULTICAST_DELEGATE_OneParam(FOnSessionStateUpdate, const EpicRtcSessionState);
 		FOnSessionStateUpdate OnSessionStateUpdate;
 
@@ -37,9 +36,9 @@ namespace UE::PixelStreaming2
 
 		DECLARE_TS_MULTICAST_DELEGATE_OneParam(FOnSessionRoomsAvailableUpdate, EpicRtcStringArrayInterface*);
 		FOnSessionRoomsAvailableUpdate OnSessionRoomsAvailableUpdate;
-		/* End FEpicRtcSessionObserver Callbacks */
+		// End FEpicRtcSessionObserver Callbacks
 
-		/* Begin FEpicRtcRoomObserver Callbacks */
+		// Begin FEpicRtcRoomObserver Callbacks
 		DECLARE_TS_MULTICAST_DELEGATE_OneParam(FOnRoomStateUpdate, const EpicRtcRoomState);
 		FOnRoomStateUpdate OnRoomStateUpdate;
 
@@ -66,9 +65,9 @@ namespace UE::PixelStreaming2
 
 		DECLARE_TS_MULTICAST_DELEGATE_OneParam(FOnRoomErrorUpdate, const EpicRtcErrorCode);
 		FOnRoomErrorUpdate OnRoomErrorUpdate;
-		/* End FEpicRtcRoomObserver Callbacks */
+		// End FEpicRtcRoomObserver Callbacks
 
-		/* Begin FEpicRtcAudioTrackObserver Callbacks */
+		// Begin FEpicRtcAudioTrackObserver Callbacks
 		DECLARE_TS_MULTICAST_DELEGATE_TwoParams(FOnAudioTrackMuted, EpicRtcAudioTrackInterface*, EpicRtcBool);
 		FOnAudioTrackMuted OnAudioTrackMuted;
 
@@ -80,9 +79,9 @@ namespace UE::PixelStreaming2
 
 		DECLARE_TS_MULTICAST_DELEGATE_TwoParams(FOnAudioTrackState, EpicRtcAudioTrackInterface*, const EpicRtcTrackState);
 		FOnAudioTrackState OnAudioTrackState;
-		/* End FEpicRtcAudioTrackObserver Callbacks */
+		// End FEpicRtcAudioTrackObserver Callbacks
 
-		/* Begin FEpicRtcVideoTrackObserver Callbacks */
+		// Begin FEpicRtcVideoTrackObserver Callbacks
 		DECLARE_TS_MULTICAST_DELEGATE_TwoParams(FOnVideoTrackMuted, EpicRtcVideoTrackInterface*, EpicRtcBool);
 		FOnVideoTrackMuted OnVideoTrackMuted;
 
@@ -94,10 +93,9 @@ namespace UE::PixelStreaming2
 
 		DECLARE_TS_MULTICAST_DELEGATE_TwoParams(FOnVideoTrackState, EpicRtcVideoTrackInterface*, const EpicRtcTrackState);
 		FOnVideoTrackState OnVideoTrackState;
+		// End FEpicRtcVideoTrackObserver Callbacks
 
-		/* End FEpicRtcVideoTrackObserver Callbacks */
-
-		/* Begin FEpicRtcDataTrackObserver Callbacks */
+		// Begin FEpicRtcDataTrackObserver Callbacks
 		DECLARE_TS_MULTICAST_DELEGATE_OneParam(FOnDataTrackRemoved, EpicRtcDataTrackInterface*);
 		FOnDataTrackRemoved OnDataTrackRemoved;
 
@@ -106,21 +104,21 @@ namespace UE::PixelStreaming2
 
 		DECLARE_TS_MULTICAST_DELEGATE_OneParam(FOnDataTrackMessage, EpicRtcDataTrackInterface*);
 		FOnDataTrackMessage OnDataTrackMessage;
-		/* End FEpicRtcDataTrackObserver Callbacks */
+		// End FEpicRtcDataTrackObserver Callbacks
 
 	public:
-		/* Begin EpicRtc Classes */
+		// Begin EpicRtc Classes
 		TRefCountPtr<EpicRtcConferenceInterface> EpicRtcConference;
 		TRefCountPtr<EpicRtcSessionInterface>	 EpicRtcSession;
 		TRefCountPtr<EpicRtcRoomInterface>		 EpicRtcRoom;
-		/* End EpicRtc Classes */
+		// End EpicRtc Classes
 
-		/* Begin EpicRtc Observers */
+		// Begin EpicRtc Observers
 		TRefCountPtr<FEpicRtcSessionObserver>			SessionObserver;
 		TRefCountPtr<FEpicRtcRoomObserver>				RoomObserver;
 		TRefCountPtr<FEpicRtcAudioTrackObserverFactory> AudioTrackObserverFactory;
 		TRefCountPtr<FEpicRtcVideoTrackObserverFactory> VideoTrackObserverFactory;
 		TRefCountPtr<FEpicRtcDataTrackObserverFactory>	DataTrackObserverFactory;
-		/* End EpicRtc Observers */
+		// End EpicRtc Observers
 	};
 } // namespace UE::PixelStreaming2
\ No newline at end of file
diff --git a/Engine/Plugins/Media/PixelStreaming2/Source/PixelStreaming2/Internal/EpicRtcMemory.h b/Engine/Plugins/Media/PixelStreaming2/Source/PixelStreaming2/Internal/EpicRtcMemory.h
deleted file mode 100644
diff --git a/Engine/Plugins/Media/PixelStreaming2/Source/PixelStreaming2/Internal/EpicRtcRoomObserver.h b/Engine/Plugins/Media/PixelStreaming2/Source/PixelStreaming2/Internal/EpicRtcRoomObserver.h
--- a/Engine/Plugins/Media/PixelStreaming2/Source/PixelStreaming2/Internal/EpicRtcRoomObserver.h
+++ b/Engine/Plugins/Media/PixelStreaming2/Source/PixelStreaming2/Internal/EpicRtcRoomObserver.h
@@ -2,23 +2,23 @@
 
 #pragma once
 
-#include "EpicRtcMemory.h"
+#include "Templates/RefCounting.h"
 #include "Templates/SharedPointer.h"
 
 #include "epic_rtc/core/room_observer.h"
 
 namespace UE::PixelStreaming2
 {
-
 	class FEpicRtcManager;
 
 	class PIXELSTREAMING2_API FEpicRtcRoomObserver : public EpicRtcRoomObserverInterface, public TRefCountingMixin<FEpicRtcRoomObserver>
 	{
 	public:
 		FEpicRtcRoomObserver(TWeakPtr<FEpicRtcManager> Manager);
+		virtual ~FEpicRtcRoomObserver() = default;
 
 	private:
-		/* Begin EpicRtcRoomObserver */
+		// Begin EpicRtcRoomObserver
 		virtual void							   OnRoomStateUpdate(const EpicRtcRoomState State) override;
 		virtual void							   OnRoomJoinedUpdate(EpicRtcParticipantInterface* Participant) override;
 		virtual void							   OnRoomLeftUpdate(const EpicRtcStringView ParticipantId) override;
@@ -28,14 +28,14 @@ namespace UE::PixelStreaming2
 		[[nodiscard]] virtual EpicRtcSdpInterface* OnLocalSdpUpdate(EpicRtcParticipantInterface* Participant, EpicRtcSdpInterface* Sdp) override;
 		[[nodiscard]] virtual EpicRtcSdpInterface* OnRemoteSdpUpdate(EpicRtcParticipantInterface* Participant, EpicRtcSdpInterface* Sdp) override;
 		virtual void							   OnRoomErrorUpdate(const EpicRtcErrorCode Error) override;
-		/* Begin EpicRtcRoomObserver */
+		// Begin EpicRtcRoomObserver
 
 	public:
-		/* Begin EpicRtcRefCountInterface */
+		// Begin EpicRtcRefCountInterface
 		virtual uint32_t AddRef() override final { return TRefCountingMixin<FEpicRtcRoomObserver>::AddRef(); }
 		virtual uint32_t Release() override final { return TRefCountingMixin<FEpicRtcRoomObserver>::Release(); }
 		virtual uint32_t Count() const override final { return TRefCountingMixin<FEpicRtcRoomObserver>::GetRefCount(); }
-		/* End EpicRtcRefCountInterface */
+		// End EpicRtcRefCountInterface
 
 	private:
 		TWeakPtr<FEpicRtcManager> Manager;
diff --git a/Engine/Plugins/Media/PixelStreaming2/Source/PixelStreaming2/Internal/EpicRtcSessionObserver.h b/Engine/Plugins/Media/PixelStreaming2/Source/PixelStreaming2/Internal/EpicRtcSessionObserver.h
--- a/Engine/Plugins/Media/PixelStreaming2/Source/PixelStreaming2/Internal/EpicRtcSessionObserver.h
+++ b/Engine/Plugins/Media/PixelStreaming2/Source/PixelStreaming2/Internal/EpicRtcSessionObserver.h
@@ -2,33 +2,33 @@
 
 #pragma once
 
-#include "EpicRtcMemory.h"
+#include "Templates/RefCounting.h"
 #include "Templates/SharedPointer.h"
 
 #include "epic_rtc/core/session_observer.h"
 
 namespace UE::PixelStreaming2
 {
-
 	class FEpicRtcManager;
 
 	class PIXELSTREAMING2_API FEpicRtcSessionObserver : public EpicRtcSessionObserverInterface, public TRefCountingMixin<FEpicRtcSessionObserver>
 	{
 	public:
 		FEpicRtcSessionObserver(TWeakPtr<FEpicRtcManager> Manager);
+		virtual ~FEpicRtcSessionObserver() = default;
 
 	private:
-		/* Begin EpicRtcSessionObserver */
+		// Begin EpicRtcSessionObserver
 		virtual void OnSessionStateUpdate(const EpicRtcSessionState State) override;
 		virtual void OnSessionErrorUpdate(const EpicRtcErrorCode Error) override;
 		virtual void OnSessionRoomsAvailableUpdate(EpicRtcStringArrayInterface* RoomsList) override;
-		/* End EpicRtcSessionObserver */
+		// End EpicRtcSessionObserver
 	public:
-		/* Begin EpicRtcRefCountInterface */
+		// Begin EpicRtcRefCountInterface
 		virtual uint32_t AddRef() override final { return TRefCountingMixin<FEpicRtcSessionObserver>::AddRef(); }
 		virtual uint32_t Release() override final { return TRefCountingMixin<FEpicRtcSessionObserver>::Release(); }
 		virtual uint32_t Count() const override final { return TRefCountingMixin<FEpicRtcSessionObserver>::GetRefCount(); }
-		/* End EpicRtcRefCountInterface */
+		// End EpicRtcRefCountInterface
 	private:
 		TWeakPtr<FEpicRtcManager> Manager;
 	};
diff --git a/Engine/Plugins/Media/PixelStreaming2/Source/PixelStreaming2/Internal/EpicRtcTickableTask.h b/Engine/Plugins/Media/PixelStreaming2/Source/PixelStreaming2/Internal/EpicRtcTickableTask.h
--- a/Engine/Plugins/Media/PixelStreaming2/Source/PixelStreaming2/Internal/EpicRtcTickableTask.h
+++ b/Engine/Plugins/Media/PixelStreaming2/Source/PixelStreaming2/Internal/EpicRtcTickableTask.h
@@ -2,11 +2,10 @@
 
 #pragma once
 
-#include "CoreMinimal.h"
+#include "Templates/UniquePtr.h"
 
 namespace UE::PixelStreaming2
 {
-
 	template <typename T>
 	using TUniqueTaskPtr = TUniquePtr<T, struct FEpicRtcTaskDeleter>;
 
diff --git a/Engine/Plugins/Media/PixelStreaming2/Source/PixelStreaming2/Internal/EpicRtcVideoCommon.h b/Engine/Plugins/Media/PixelStreaming2/Source/PixelStreaming2/Internal/EpicRtcVideoCommon.h
--- a/Engine/Plugins/Media/PixelStreaming2/Source/PixelStreaming2/Internal/EpicRtcVideoCommon.h
+++ b/Engine/Plugins/Media/PixelStreaming2/Source/PixelStreaming2/Internal/EpicRtcVideoCommon.h
@@ -11,9 +11,13 @@
 #include "epic_rtc/core/video/video_buffer.h"
 #include "epic_rtc/core/video/video_codec_info.h"
 
-namespace UE::PixelStreaming2
+FORCEINLINE bool operator==(const EpicRtcVideoResolution& Lhs, const EpicRtcVideoResolution& Rhs)
 {
+	return Lhs._width == Rhs._width && Lhs._height == Rhs._height;
+}
 
+namespace UE::PixelStreaming2
+{
 	class PIXELSTREAMING2_API FEpicRtcEncodedVideoBuffer : public EpicRtcEncodedVideoBufferInterface, public TRefCountingMixin<FEpicRtcEncodedVideoBuffer>
 	{
 	public:
@@ -31,11 +35,11 @@ namespace UE::PixelStreaming2
 		TArray<uint8> Data;
 
 	public:
-		/* Begin EpicRtcRefCountInterface */
+		// Begin EpicRtcRefCountInterface
 		virtual uint32_t AddRef() override final { return TRefCountingMixin<FEpicRtcEncodedVideoBuffer>::AddRef(); }
 		virtual uint32_t Release() override final { return TRefCountingMixin<FEpicRtcEncodedVideoBuffer>::Release(); }
 		virtual uint32_t Count() const override final { return TRefCountingMixin<FEpicRtcEncodedVideoBuffer>::GetRefCount(); }
-		/* End EpicRtcRefCountInterface */
+		// End EpicRtcRefCountInterface
 	};
 
 	class PIXELSTREAMING2_API FEpicRtcParameterPairArray : public EpicRtcParameterPairArrayInterface, public TRefCountingMixin<FEpicRtcParameterPairArray>
@@ -78,11 +82,11 @@ namespace UE::PixelStreaming2
 		TArray<EpicRtcParameterPair> Data;
 
 	public:
-		/* Begin EpicRtcRefCountInterface */
+		// Begin EpicRtcRefCountInterface
 		virtual uint32_t AddRef() override final { return TRefCountingMixin<FEpicRtcParameterPairArray>::AddRef(); }
 		virtual uint32_t Release() override final { return TRefCountingMixin<FEpicRtcParameterPairArray>::Release(); }
 		virtual uint32_t Count() const override final { return TRefCountingMixin<FEpicRtcParameterPairArray>::GetRefCount(); }
-		/* End EpicRtcRefCountInterface */
+		// End EpicRtcRefCountInterface
 	};
 
 	class PIXELSTREAMING2_API FEpicRtcScalabilityModeArray : public EpicRtcVideoScalabilityModeArrayInterface, public TRefCountingMixin<FEpicRtcScalabilityModeArray>
@@ -133,17 +137,17 @@ namespace UE::PixelStreaming2
 		TArray<EpicRtcVideoScalabilityMode> Data;
 
 	public:
-		/* Begin EpicRtcRefCountInterface */
+		// Begin EpicRtcRefCountInterface
 		virtual uint32_t AddRef() override final { return TRefCountingMixin<FEpicRtcScalabilityModeArray>::AddRef(); }
 		virtual uint32_t Release() override final { return TRefCountingMixin<FEpicRtcScalabilityModeArray>::Release(); }
 		virtual uint32_t Count() const override final { return TRefCountingMixin<FEpicRtcScalabilityModeArray>::GetRefCount(); }
-		/* End EpicRtcRefCountInterface */
+		// End EpicRtcRefCountInterface
 	};
 
-	class PIXELSTREAMING2_API FVideoCodecInfo : public EpicRtcVideoCodecInfoInterface, public TRefCountingMixin<FVideoCodecInfo>
+	class PIXELSTREAMING2_API FEpicRtcVideoCodecInfo : public EpicRtcVideoCodecInfoInterface, public TRefCountingMixin<FEpicRtcVideoCodecInfo>
 	{
 	public:
-		FVideoCodecInfo(EpicRtcVideoCodec Codec, bool bIsHardwareAccelerated = false, EpicRtcParameterPairArrayInterface* Parameters = new FEpicRtcParameterPairArray(), EpicRtcVideoScalabilityModeArrayInterface* ScalabilityModes = new FEpicRtcScalabilityModeArray())
+		FEpicRtcVideoCodecInfo(EpicRtcVideoCodec Codec, bool bIsHardwareAccelerated = false, EpicRtcParameterPairArrayInterface* Parameters = new FEpicRtcParameterPairArray(), EpicRtcVideoScalabilityModeArrayInterface* ScalabilityModes = new FEpicRtcScalabilityModeArray())
 			: Codec(Codec)
 			, bIsHardwareAccelerated(bIsHardwareAccelerated)
 			, Parameters(Parameters)
@@ -151,7 +155,7 @@ namespace UE::PixelStreaming2
 		{
 		}
 
-		virtual ~FVideoCodecInfo() = default;
+		virtual ~FEpicRtcVideoCodecInfo() = default;
 
 		EpicRtcVideoCodec GetCodec() override
 		{
@@ -180,11 +184,11 @@ namespace UE::PixelStreaming2
 		TRefCountPtr<EpicRtcVideoScalabilityModeArrayInterface> ScalabilityModes;
 
 	public:
-		/* Begin EpicRtcRefCountInterface */
-		virtual uint32_t AddRef() override final { return TRefCountingMixin<FVideoCodecInfo>::AddRef(); }
-		virtual uint32_t Release() override final { return TRefCountingMixin<FVideoCodecInfo>::Release(); }
-		virtual uint32_t Count() const override final { return TRefCountingMixin<FVideoCodecInfo>::GetRefCount(); }
-		/* End EpicRtcRefCountInterface */
+		// Begin EpicRtcRefCountInterface
+		virtual uint32_t AddRef() override final { return TRefCountingMixin<FEpicRtcVideoCodecInfo>::AddRef(); }
+		virtual uint32_t Release() override final { return TRefCountingMixin<FEpicRtcVideoCodecInfo>::Release(); }
+		virtual uint32_t Count() const override final { return TRefCountingMixin<FEpicRtcVideoCodecInfo>::GetRefCount(); }
+		// End EpicRtcRefCountInterface
 	};
 
 	class PIXELSTREAMING2_API FVideoCodecInfoArray : public EpicRtcVideoCodecInfoArrayInterface, public TRefCountingMixin<FVideoCodecInfoArray>
@@ -258,11 +262,11 @@ namespace UE::PixelStreaming2
 		TArray<EpicRtcVideoCodecInfoInterface*> Data;
 
 	public:
-		/* Begin EpicRtcRefCountInterface */
+		// Begin EpicRtcRefCountInterface
 		virtual uint32_t AddRef() override final { return TRefCountingMixin<FVideoCodecInfoArray>::AddRef(); }
 		virtual uint32_t Release() override final { return TRefCountingMixin<FVideoCodecInfoArray>::Release(); }
 		virtual uint32_t Count() const override final { return TRefCountingMixin<FVideoCodecInfoArray>::GetRefCount(); }
-		/* End EpicRtcRefCountInterface */
+		// End EpicRtcRefCountInterface
 	};
 
 	class PIXELSTREAMING2_API FEpicRtcVideoResolutionBitrateLimitsArray : public EpicRtcVideoResolutionBitrateLimitsArrayInterface, public TRefCountingMixin<FEpicRtcVideoResolutionBitrateLimitsArray>
@@ -300,11 +304,11 @@ namespace UE::PixelStreaming2
 		TArray<EpicRtcVideoResolutionBitrateLimits> Data;
 
 	public:
-		/* Begin EpicRtcRefCountInterface */
+		// Begin EpicRtcRefCountInterface
 		virtual uint32_t AddRef() override final { return TRefCountingMixin<FEpicRtcVideoResolutionBitrateLimitsArray>::AddRef(); }
 		virtual uint32_t Release() override final { return TRefCountingMixin<FEpicRtcVideoResolutionBitrateLimitsArray>::Release(); }
 		virtual uint32_t Count() const override final { return TRefCountingMixin<FEpicRtcVideoResolutionBitrateLimitsArray>::GetRefCount(); }
-		/* End EpicRtcRefCountInterface */
+		// End EpicRtcRefCountInterface
 	};
 
 	class PIXELSTREAMING2_API FEpicRtcPixelFormatArray : public EpicRtcPixelFormatArrayInterface, public TRefCountingMixin<FEpicRtcPixelFormatArray>
@@ -342,11 +346,11 @@ namespace UE::PixelStreaming2
 		TArray<EpicRtcPixelFormat> Data;
 
 	public:
-		/* Begin EpicRtcRefCountInterface */
+		// Begin EpicRtcRefCountInterface
 		virtual uint32_t AddRef() override final { return TRefCountingMixin<FEpicRtcPixelFormatArray>::AddRef(); }
 		virtual uint32_t Release() override final { return TRefCountingMixin<FEpicRtcPixelFormatArray>::Release(); }
 		virtual uint32_t Count() const override final { return TRefCountingMixin<FEpicRtcPixelFormatArray>::GetRefCount(); }
-		/* End EpicRtcRefCountInterface */
+		// End EpicRtcRefCountInterface
 	};
 
 	class PIXELSTREAMING2_API FEpicRtcVideoFrameTypeArray : public EpicRtcVideoFrameTypeArrayInterface, public TRefCountingMixin<FEpicRtcVideoFrameTypeArray>
@@ -379,11 +383,11 @@ namespace UE::PixelStreaming2
 		TArray<EpicRtcVideoFrameType> Data;
 
 	public:
-		/* Begin EpicRtcRefCountInterface */
+		// Begin EpicRtcRefCountInterface
 		virtual uint32_t AddRef() override final { return TRefCountingMixin<FEpicRtcVideoFrameTypeArray>::AddRef(); }
 		virtual uint32_t Release() override final { return TRefCountingMixin<FEpicRtcVideoFrameTypeArray>::Release(); }
 		virtual uint32_t Count() const override final { return TRefCountingMixin<FEpicRtcVideoFrameTypeArray>::GetRefCount(); }
-		/* End EpicRtcRefCountInterface */
+		// End EpicRtcRefCountInterface
 	};
 
 	class PIXELSTREAMING2_API FEpicRtcInt32Array : public EpicRtcInt32ArrayInterface, public TRefCountingMixin<FEpicRtcInt32Array>
@@ -426,11 +430,11 @@ namespace UE::PixelStreaming2
 		TArray<int32_t> Data;
 
 	public:
-		/* Begin EpicRtcRefCountInterface */
+		// Begin EpicRtcRefCountInterface
 		virtual uint32_t AddRef() override final { return TRefCountingMixin<FEpicRtcInt32Array>::AddRef(); }
 		virtual uint32_t Release() override final { return TRefCountingMixin<FEpicRtcInt32Array>::Release(); }
 		virtual uint32_t Count() const override final { return TRefCountingMixin<FEpicRtcInt32Array>::GetRefCount(); }
-		/* End EpicRtcRefCountInterface */
+		// End EpicRtcRefCountInterface
 	};
 
 	class PIXELSTREAMING2_API FEpicRtcBoolArray : public EpicRtcBoolArrayInterface, public TRefCountingMixin<FEpicRtcBoolArray>
@@ -485,11 +489,11 @@ namespace UE::PixelStreaming2
 		TArray<EpicRtcBool> Data;
 
 	public:
-		/* Begin EpicRtcRefCountInterface */
+		// Begin EpicRtcRefCountInterface
 		virtual uint32_t AddRef() override final { return TRefCountingMixin<FEpicRtcBoolArray>::AddRef(); }
 		virtual uint32_t Release() override final { return TRefCountingMixin<FEpicRtcBoolArray>::Release(); }
 		virtual uint32_t Count() const override final { return TRefCountingMixin<FEpicRtcBoolArray>::GetRefCount(); }
-		/* End EpicRtcRefCountInterface */
+		// End EpicRtcRefCountInterface
 	};
 
 	class PIXELSTREAMING2_API FEpicRtcDecodeTargetIndicationArray : public EpicRtcDecodeTargetIndicationArrayInterface, public TRefCountingMixin<FEpicRtcDecodeTargetIndicationArray>
@@ -554,11 +558,11 @@ namespace UE::PixelStreaming2
 		TArray<EpicRtcDecodeTargetIndication> Data;
 
 	public:
-		/* Begin EpicRtcRefCountInterface */
+		// Begin EpicRtcRefCountInterface
 		virtual uint32_t AddRef() override final { return TRefCountingMixin<FEpicRtcDecodeTargetIndicationArray>::AddRef(); }
 		virtual uint32_t Release() override final { return TRefCountingMixin<FEpicRtcDecodeTargetIndicationArray>::Release(); }
 		virtual uint32_t Count() const override final { return TRefCountingMixin<FEpicRtcDecodeTargetIndicationArray>::GetRefCount(); }
-		/* End EpicRtcRefCountInterface */
+		// End EpicRtcRefCountInterface
 	};
 
 	class PIXELSTREAMING2_API FEpicRtcCodecBufferUsageArray : public EpicRtcCodecBufferUsageArrayInterface, public TRefCountingMixin<FEpicRtcCodecBufferUsageArray>
@@ -611,11 +615,11 @@ namespace UE::PixelStreaming2
 		TArray<EpicRtcCodecBufferUsage> Data;
 
 	public:
-		/* Begin EpicRtcRefCountInterface */
+		// Begin EpicRtcRefCountInterface
 		virtual uint32_t AddRef() override final { return TRefCountingMixin<FEpicRtcCodecBufferUsageArray>::AddRef(); }
 		virtual uint32_t Release() override final { return TRefCountingMixin<FEpicRtcCodecBufferUsageArray>::Release(); }
 		virtual uint32_t Count() const override final { return TRefCountingMixin<FEpicRtcCodecBufferUsageArray>::GetRefCount(); }
-		/* End EpicRtcRefCountInterface */
+		// End EpicRtcRefCountInterface
 	};
 
 	class PIXELSTREAMING2_API FEpicRtcVideoResolutionArray : public EpicRtcVideoResolutionArrayInterface, public TRefCountingMixin<FEpicRtcVideoResolutionArray>
@@ -667,11 +671,11 @@ namespace UE::PixelStreaming2
 		TArray<EpicRtcVideoResolution> Data;
 
 	public:
-		/* Begin EpicRtcRefCountInterface */
+		// Begin EpicRtcRefCountInterface
 		virtual uint32_t AddRef() override final { return TRefCountingMixin<FEpicRtcVideoResolutionArray>::AddRef(); }
 		virtual uint32_t Release() override final { return TRefCountingMixin<FEpicRtcVideoResolutionArray>::Release(); }
 		virtual uint32_t Count() const override final { return TRefCountingMixin<FEpicRtcVideoResolutionArray>::GetRefCount(); }
-		/* End EpicRtcRefCountInterface */
+		// End EpicRtcRefCountInterface
 	};
 
 	class PIXELSTREAMING2_API FEpicRtcGenericFrameInfoArray : public EpicRtcGenericFrameInfoArrayInterface, public TRefCountingMixin<FEpicRtcGenericFrameInfoArray>
@@ -745,11 +749,125 @@ namespace UE::PixelStreaming2
 		TArray<EpicRtcGenericFrameInfoInterface*> Data;
 
 	public:
-		/* Begin EpicRtcRefCountInterface */
+		// Begin EpicRtcRefCountInterface
 		virtual uint32_t AddRef() override final { return TRefCountingMixin<FEpicRtcGenericFrameInfoArray>::AddRef(); }
 		virtual uint32_t Release() override final { return TRefCountingMixin<FEpicRtcGenericFrameInfoArray>::Release(); }
 		virtual uint32_t Count() const override final { return TRefCountingMixin<FEpicRtcGenericFrameInfoArray>::GetRefCount(); }
-		/* End EpicRtcRefCountInterface */
+		// End EpicRtcRefCountInterface
+	};
+
+	class PIXELSTREAMING2_API FEpicRtcGenericFrameInfo : public EpicRtcGenericFrameInfoInterface, public TRefCountingMixin<FEpicRtcGenericFrameInfo>
+	{
+	public:
+		FEpicRtcGenericFrameInfo(const FGenericFrameInfo& GenericFrameInfo)
+			: SpatialId(GenericFrameInfo.SpatialId)
+			, TemporalId(GenericFrameInfo.TemporalId)
+			, DecodeTargetIndications(MakeRefCount<FEpicRtcDecodeTargetIndicationArray>(GenericFrameInfo.DecodeTargetIndications))
+			, FrameDiffs(MakeRefCount<FEpicRtcInt32Array>(GenericFrameInfo.FrameDiffs))
+			, ChainDiffs(MakeRefCount<FEpicRtcInt32Array>(GenericFrameInfo.ChainDiffs))
+			, EncoderBuffers(MakeRefCount<FEpicRtcCodecBufferUsageArray>(GenericFrameInfo.EncoderBuffers))
+			, PartOfChain(MakeRefCount<FEpicRtcBoolArray>(GenericFrameInfo.PartOfChain))
+			, ActiveDecodeTargets(MakeRefCount<FEpicRtcBoolArray>(GenericFrameInfo.ActiveDecodeTargets))
+		{
+		}
+
+		virtual ~FEpicRtcGenericFrameInfo() = default;
+
+		virtual int32_t										 GetSpatialLayerId() override { return SpatialId; }
+		virtual int32_t										 GetTemporalLayerId() override { return TemporalId; }
+		virtual EpicRtcDecodeTargetIndicationArrayInterface* GetDecodeTargetIndications() override { return DecodeTargetIndications; }
+		virtual EpicRtcInt32ArrayInterface*					 GetFrameDiffs() override { return FrameDiffs; }
+		virtual EpicRtcInt32ArrayInterface*					 GetChainDiffs() override { return ChainDiffs; }
+		virtual EpicRtcCodecBufferUsageArrayInterface*		 GetEncoderBufferUsages() override { return EncoderBuffers; }
+		virtual EpicRtcBoolArrayInterface*					 GetPartOfChain() override { return PartOfChain; }
+		virtual EpicRtcBoolArrayInterface*					 GetActiveDecodeTargets() override { return ActiveDecodeTargets; }
+
+	private:
+		int32_t											  SpatialId;
+		int32_t											  TemporalId;
+		TRefCountPtr<FEpicRtcDecodeTargetIndicationArray> DecodeTargetIndications;
+		TRefCountPtr<FEpicRtcInt32Array>				  FrameDiffs;
+		TRefCountPtr<FEpicRtcInt32Array>				  ChainDiffs;
+		TRefCountPtr<FEpicRtcCodecBufferUsageArray>		  EncoderBuffers;
+		TRefCountPtr<FEpicRtcBoolArray>					  PartOfChain;
+		TRefCountPtr<FEpicRtcBoolArray>					  ActiveDecodeTargets;
+
+	public:
+		// Begin EpicRtcRefCountInterface
+		virtual uint32_t AddRef() override final { return TRefCountingMixin<FEpicRtcGenericFrameInfo>::AddRef(); }
+		virtual uint32_t Release() override final { return TRefCountingMixin<FEpicRtcGenericFrameInfo>::Release(); }
+		virtual uint32_t Count() const override final { return TRefCountingMixin<FEpicRtcGenericFrameInfo>::GetRefCount(); }
+		// End EpicRtcRefCountInterface
 	};
 
+	class PIXELSTREAMING2_API FEpicRtcFrameDependencyStructure : public EpicRtcFrameDependencyStructure, public TRefCountingMixin<FEpicRtcFrameDependencyStructure>
+	{
+	public:
+		FEpicRtcFrameDependencyStructure(const FFrameDependencyStructure& FrameDependencyStructure)
+			: StructureId(FrameDependencyStructure.StructureId)
+			, NumDecodeTargets(FrameDependencyStructure.NumDecodeTargets)
+			, NumChains(FrameDependencyStructure.NumChains)
+			, DecodeTargetProtectedByChain(MakeRefCount<FEpicRtcInt32Array>(FrameDependencyStructure.DecodeTargetProtectedByChain))
+			, Resolutions(MakeRefCount<FEpicRtcVideoResolutionArray>(FrameDependencyStructure.Resolutions))
+		{
+			TArray<EpicRtcGenericFrameInfoInterface*> GenericFrameInfoArray;
+			GenericFrameInfoArray.SetNum(FrameDependencyStructure.Templates.Num());
+
+			for (size_t i = 0; i < FrameDependencyStructure.Templates.Num(); i++)
+			{
+				FGenericFrameInfo GenericFrameInfo;
+				GenericFrameInfo.SpatialId = FrameDependencyStructure.Templates[i].SpatialId;
+				GenericFrameInfo.TemporalId = FrameDependencyStructure.Templates[i].TemporalId;
+				GenericFrameInfo.DecodeTargetIndications = FrameDependencyStructure.Templates[i].DecodeTargetIndications;
+				GenericFrameInfo.FrameDiffs = FrameDependencyStructure.Templates[i].FrameDiffs;
+				GenericFrameInfo.ChainDiffs = FrameDependencyStructure.Templates[i].ChainDiffs;
+
+				GenericFrameInfoArray[i] = new FEpicRtcGenericFrameInfo(GenericFrameInfo);
+			}
+
+			Templates = MakeRefCount<FEpicRtcGenericFrameInfoArray>(GenericFrameInfoArray);
+		}
+
+		virtual ~FEpicRtcFrameDependencyStructure() = default;
+
+		virtual int32_t								   GetStructureId() override { return StructureId; }
+		virtual int32_t								   GetNumDecodeTargets() override { return NumDecodeTargets; }
+		virtual int32_t								   GetNumChains() override { return NumChains; }
+		virtual EpicRtcInt32ArrayInterface*			   GetDecodeTargetProtectedByChain() override { return DecodeTargetProtectedByChain; }
+		virtual EpicRtcVideoResolutionArrayInterface*  GetResolutions() override { return Resolutions; }
+		virtual EpicRtcGenericFrameInfoArrayInterface* GetTemplates() override { return Templates; }
+
+		friend bool operator==(FEpicRtcFrameDependencyStructure& Lhs, FEpicRtcFrameDependencyStructure& Rhs)
+		{
+			TArray<int32_t> LhsDecodeTargetProtectedByChain(Lhs.GetDecodeTargetProtectedByChain()->Get(), Lhs.GetDecodeTargetProtectedByChain()->Size());
+			TArray<int32_t> RhsDecodeTargetProtectedByChain(Rhs.GetDecodeTargetProtectedByChain()->Get(), Rhs.GetDecodeTargetProtectedByChain()->Size());
+
+			TArray<EpicRtcVideoResolution> LhsResolutions(Lhs.GetResolutions()->Get(), Lhs.GetResolutions()->Size());
+			TArray<EpicRtcVideoResolution> RhsResolutions(Rhs.GetResolutions()->Get(), Rhs.GetResolutions()->Size());
+
+			TArray<EpicRtcGenericFrameInfoInterface*> LhsTemplates(Lhs.GetTemplates()->Get(), Lhs.GetTemplates()->Size());
+			TArray<EpicRtcGenericFrameInfoInterface*> RhsTemplates(Rhs.GetTemplates()->Get(), Rhs.GetTemplates()->Size());
+
+			return Lhs.NumDecodeTargets == Rhs.NumDecodeTargets
+				&& Lhs.NumChains == Rhs.NumChains
+				&& LhsDecodeTargetProtectedByChain == RhsDecodeTargetProtectedByChain
+				&& LhsResolutions == RhsResolutions
+				&& LhsTemplates == RhsTemplates;
+		}
+
+	private:
+		int											StructureId;
+		int											NumDecodeTargets;
+		int											NumChains;
+		TRefCountPtr<FEpicRtcInt32Array>			DecodeTargetProtectedByChain;
+		TRefCountPtr<FEpicRtcVideoResolutionArray>	Resolutions;
+		TRefCountPtr<FEpicRtcGenericFrameInfoArray> Templates;
+
+	public:
+		// Begin EpicRtcRefCountInterface
+		virtual uint32_t AddRef() override final { return TRefCountingMixin<FEpicRtcFrameDependencyStructure>::AddRef(); }
+		virtual uint32_t Release() override final { return TRefCountingMixin<FEpicRtcFrameDependencyStructure>::Release(); }
+		virtual uint32_t Count() const override final { return TRefCountingMixin<FEpicRtcFrameDependencyStructure>::GetRefCount(); }
+		// End EpicRtcRefCountInterface
+	};
 } // namespace UE::PixelStreaming2
\ No newline at end of file
diff --git a/Engine/Plugins/Media/PixelStreaming2/Source/PixelStreaming2/Internal/EpicRtcVideoDecoderInitializer.h b/Engine/Plugins/Media/PixelStreaming2/Source/PixelStreaming2/Internal/EpicRtcVideoDecoderInitializer.h
--- a/Engine/Plugins/Media/PixelStreaming2/Source/PixelStreaming2/Internal/EpicRtcVideoDecoderInitializer.h
+++ b/Engine/Plugins/Media/PixelStreaming2/Source/PixelStreaming2/Internal/EpicRtcVideoDecoderInitializer.h
@@ -3,8 +3,6 @@
 #pragma once
 
 #include "EpicRtcVideoCommon.h"
-
-#include "Containers/StaticArray.h"
 #include "Video/VideoConfig.h"
 
 #include "epic_rtc/core/video/video_decoder.h"
@@ -18,25 +16,21 @@ namespace UE::PixelStreaming2
 		FEpicRtcVideoDecoderInitializer() = default;
 		virtual ~FEpicRtcVideoDecoderInitializer() = default;
 
-		/* Begin EpicRtcVideoDecoderInitializerInterface */
+		// Begin EpicRtcVideoDecoderInitializerInterface
 		virtual void								 CreateDecoder(EpicRtcVideoCodecInfoInterface* CodecInfo, EpicRtcVideoDecoderInterface** OutDecoder) override;
 		virtual EpicRtcStringView					 GetName() override;
 		virtual EpicRtcVideoCodecInfoArrayInterface* GetSupportedCodecs() override;
-		/* End EpicRtcVideoDecoderInitializerInterface */
-
-	private:
-		// The list of each individual codec we have support for (order of this array is preference order after selected codec)
-		static TStaticArray<EVideoCodec, 4> SupportedCodecList;
+		// End EpicRtcVideoDecoderInitializerInterface
 
 	private:
 		TMap<EVideoCodec, TArray<TRefCountPtr<EpicRtcVideoCodecInfoInterface>>> CreateSupportedDecoderMap();
 
 	public:
-		/* Begin EpicRtcRefCountInterface */
+		// Begin EpicRtcRefCountInterface
 		virtual uint32_t AddRef() override final { return TRefCountingMixin<FEpicRtcVideoDecoderInitializer>::AddRef(); }
 		virtual uint32_t Release() override final { return TRefCountingMixin<FEpicRtcVideoDecoderInitializer>::Release(); }
 		virtual uint32_t Count() const override final { return TRefCountingMixin<FEpicRtcVideoDecoderInitializer>::GetRefCount(); }
-		/* End EpicRtcRefCountInterface */
+		// End EpicRtcRefCountInterface
 	};
 
 } // namespace UE::PixelStreaming2
\ No newline at end of file
diff --git a/Engine/Plugins/Media/PixelStreaming2/Source/PixelStreaming2/Internal/EpicRtcVideoEncoderInitializer.h b/Engine/Plugins/Media/PixelStreaming2/Source/PixelStreaming2/Internal/EpicRtcVideoEncoderInitializer.h
--- a/Engine/Plugins/Media/PixelStreaming2/Source/PixelStreaming2/Internal/EpicRtcVideoEncoderInitializer.h
+++ b/Engine/Plugins/Media/PixelStreaming2/Source/PixelStreaming2/Internal/EpicRtcVideoEncoderInitializer.h
@@ -3,8 +3,6 @@
 #pragma once
 
 #include "EpicRtcVideoCommon.h"
-
-#include "Containers/StaticArray.h"
 #include "Video/VideoConfig.h"
 
 #include "epic_rtc/core/video/video_encoder.h"
@@ -18,25 +16,21 @@ namespace UE::PixelStreaming2
 		FEpicRtcVideoEncoderInitializer() = default;
 		virtual ~FEpicRtcVideoEncoderInitializer() = default;
 
-		/* Begin EpicRtcVideoEncoderInitializerInterface */
+		// Begin EpicRtcVideoEncoderInitializerInterface
 		virtual void								 CreateEncoder(EpicRtcVideoCodecInfoInterface* CodecInfo, EpicRtcVideoEncoderInterface** OutEncoder) override;
 		virtual EpicRtcStringView					 GetName() override;
 		virtual EpicRtcVideoCodecInfoArrayInterface* GetSupportedCodecs() override;
-		/* End EpicRtcVideoEncoderInitializerInterface */
-
-	private:
-		// The list of each individual codec we have support for (order of this array is preference order after selected codec)
-		static TStaticArray<EVideoCodec, 4> SupportedCodecList;
+		// End EpicRtcVideoEncoderInitializerInterface
 
 	private:
 		TMap<EVideoCodec, TArray<TRefCountPtr<EpicRtcVideoCodecInfoInterface>>> CreateSupportedEncoderMap();
 
 	public:
-		/* Begin EpicRtcRefCountInterface */
+		// Begin EpicRtcRefCountInterface
 		virtual uint32_t AddRef() override final { return TRefCountingMixin<FEpicRtcVideoEncoderInitializer>::AddRef(); }
 		virtual uint32_t Release() override final { return TRefCountingMixin<FEpicRtcVideoEncoderInitializer>::Release(); }
 		virtual uint32_t Count() const override final { return TRefCountingMixin<FEpicRtcVideoEncoderInitializer>::GetRefCount(); }
-		/* End EpicRtcRefCountInterface */
+		// End EpicRtcRefCountInterface
 	};
 
 } // namespace UE::PixelStreaming2
\ No newline at end of file
diff --git a/Engine/Plugins/Media/PixelStreaming2/Source/PixelStreaming2/Internal/EpicRtcVideoTrackObserver.h b/Engine/Plugins/Media/PixelStreaming2/Source/PixelStreaming2/Internal/EpicRtcVideoTrackObserver.h
--- a/Engine/Plugins/Media/PixelStreaming2/Source/PixelStreaming2/Internal/EpicRtcVideoTrackObserver.h
+++ b/Engine/Plugins/Media/PixelStreaming2/Source/PixelStreaming2/Internal/EpicRtcVideoTrackObserver.h
@@ -2,7 +2,6 @@
 
 #pragma once
 
-#include "EpicRtcMemory.h"
 #include "Templates/RefCounting.h"
 #include "Templates/SharedPointer.h"
 
@@ -10,30 +9,30 @@
 
 namespace UE::PixelStreaming2
 {
-
 	class FEpicRtcManager;
 
-	class PIXELSTREAMING2_API FEpicRtcVideoTrackObserver : public EpicRtcVideoTrackObserverInterface, public TEpicRtcRefCountPtr<FEpicRtcVideoTrackObserver>
+	class PIXELSTREAMING2_API FEpicRtcVideoTrackObserver : public EpicRtcVideoTrackObserverInterface, public TRefCountingMixin<FEpicRtcVideoTrackObserver>
 	{
 	public:
 		FEpicRtcVideoTrackObserver(TWeakPtr<FEpicRtcManager> Manager);
+		virtual ~FEpicRtcVideoTrackObserver() = default;
 
 	private:
-		/* Begin EpicRtcVideoTrackObserverInterface */
-		virtual void OnVideoTrackMuted(EpicRtcVideoTrackInterface* VideoTrack, EpicRtcBool bIsMuted) override;
-		virtual void OnVideoTrackFrame(EpicRtcVideoTrackInterface* VideoTrack, const EpicRtcVideoFrame& Frame) override;
-		virtual void OnVideoTrackRemoved(EpicRtcVideoTrackInterface* VideoTrack) override;
-		virtual void OnVideoTrackState(EpicRtcVideoTrackInterface* VideoTrack, const EpicRtcTrackState State) override;
-		virtual void OnVideoTrackEncodedFrame(EpicRtcVideoTrackInterface*, const EpicRtcEncodedVideoFrame&) override {};
+		// Begin EpicRtcVideoTrackObserverInterface
+		virtual void		OnVideoTrackMuted(EpicRtcVideoTrackInterface* VideoTrack, EpicRtcBool bIsMuted) override;
+		virtual void		OnVideoTrackFrame(EpicRtcVideoTrackInterface* VideoTrack, const EpicRtcVideoFrame& Frame) override;
+		virtual void		OnVideoTrackRemoved(EpicRtcVideoTrackInterface* VideoTrack) override;
+		virtual void		OnVideoTrackState(EpicRtcVideoTrackInterface* VideoTrack, const EpicRtcTrackState State) override;
+		virtual void		OnVideoTrackEncodedFrame(EpicRtcVideoTrackInterface*, const EpicRtcEncodedVideoFrame&) override {};
 		virtual EpicRtcBool Enabled() const override { return true; };
-		/* End EpicRtcVideoTrackObserverInterface */
+		// End EpicRtcVideoTrackObserverInterface
 
 	public:
-		/* Begin EpicRtcRefCountInterface */
-		virtual uint32_t AddRef() override final { return TEpicRtcRefCountPtr<FEpicRtcVideoTrackObserver>::AddRef(); }
-		virtual uint32_t Release() override final { return TEpicRtcRefCountPtr<FEpicRtcVideoTrackObserver>::Release(); }
-		virtual uint32_t Count() const override final { return TEpicRtcRefCountPtr<FEpicRtcVideoTrackObserver>::GetRefCount(); }
-		/* End EpicRtcRefCountInterface */
+		// Begin EpicRtcRefCountInterface
+		virtual uint32_t AddRef() override final { return TRefCountingMixin<FEpicRtcVideoTrackObserver>::AddRef(); }
+		virtual uint32_t Release() override final { return TRefCountingMixin<FEpicRtcVideoTrackObserver>::Release(); }
+		virtual uint32_t Count() const override final { return TRefCountingMixin<FEpicRtcVideoTrackObserver>::GetRefCount(); }
+		// End EpicRtcRefCountInterface
 
 	private:
 		TWeakPtr<FEpicRtcManager> Manager;
diff --git a/Engine/Plugins/Media/PixelStreaming2/Source/PixelStreaming2/Internal/EpicRtcVideoTrackObserverFactory.h b/Engine/Plugins/Media/PixelStreaming2/Source/PixelStreaming2/Internal/EpicRtcVideoTrackObserverFactory.h
--- a/Engine/Plugins/Media/PixelStreaming2/Source/PixelStreaming2/Internal/EpicRtcVideoTrackObserverFactory.h
+++ b/Engine/Plugins/Media/PixelStreaming2/Source/PixelStreaming2/Internal/EpicRtcVideoTrackObserverFactory.h
@@ -2,7 +2,6 @@
 
 #pragma once
 
-#include "EpicRtcMemory.h"
 #include "Templates/RefCounting.h"
 #include "Templates/SharedPointer.h"
 
@@ -10,25 +9,25 @@
 
 namespace UE::PixelStreaming2
 {
-
 	class FEpicRtcManager;
 
 	class PIXELSTREAMING2_API FEpicRtcVideoTrackObserverFactory : public EpicRtcVideoTrackObserverFactoryInterface, public TRefCountingMixin<FEpicRtcVideoTrackObserverFactory>
 	{
 	public:
 		FEpicRtcVideoTrackObserverFactory(TWeakPtr<FEpicRtcManager> Manager);
+		virtual ~FEpicRtcVideoTrackObserverFactory() = default;
 
 	public:
-		/* Begin EpicRtcVideoTrackObserverFactoryInterface */
+		// Begin EpicRtcVideoTrackObserverFactoryInterface
 		virtual EpicRtcErrorCode CreateVideoTrackObserver(const EpicRtcStringView ParticipantId, const EpicRtcStringView VideoTrackId, EpicRtcVideoTrackObserverInterface** OutVideoTrackObserver) override;
-		/* End EpicRtcVideoTrackObserverFactoryInterface */
+		// End EpicRtcVideoTrackObserverFactoryInterface
 
 	public:
-		/* Begin EpicRtcRefCountInterface */
+		// Begin EpicRtcRefCountInterface
 		virtual uint32_t AddRef() override final { return TRefCountingMixin<FEpicRtcVideoTrackObserverFactory>::AddRef(); }
 		virtual uint32_t Release() override final { return TRefCountingMixin<FEpicRtcVideoTrackObserverFactory>::Release(); }
 		virtual uint32_t Count() const override final { return TRefCountingMixin<FEpicRtcVideoTrackObserverFactory>::GetRefCount(); }
-		/* End EpicRtcRefCountInterface */
+		// End EpicRtcRefCountInterface
 
 	private:
 		TWeakPtr<FEpicRtcManager> Manager;
diff --git a/Engine/Plugins/Media/PixelStreaming2/Source/PixelStreaming2/Internal/EpicRtcWebsocket.h b/Engine/Plugins/Media/PixelStreaming2/Source/PixelStreaming2/Internal/EpicRtcWebsocket.h
--- a/Engine/Plugins/Media/PixelStreaming2/Source/PixelStreaming2/Internal/EpicRtcWebsocket.h
+++ b/Engine/Plugins/Media/PixelStreaming2/Source/PixelStreaming2/Internal/EpicRtcWebsocket.h
@@ -3,33 +3,32 @@
 #pragma once
 
 #include "Engine/EngineTypes.h"
-#include "EpicRtcMemory.h"
 #include "Tickable.h"
+
 #include "epic_rtc/plugins/signalling/websocket.h"
 
 class IWebSocket;
 
 namespace UE::PixelStreaming2
 {
-
-	class PIXELSTREAMING2_API FEpicRtcWebsocket : public EpicRtcWebsocketInterface, public TRefCountingMixin<FEpicRtcWebsocket>, public FTickableGameObject
+	class PIXELSTREAMING2_API FEpicRtcWebsocket : public EpicRtcWebsocketInterface, public FTickableGameObject, public TRefCountingMixin<FEpicRtcWebsocket>
 	{
 	public:
 		FEpicRtcWebsocket(bool bKeepAlive = true, TSharedPtr<IWebSocket> WebSocket = nullptr);
 		virtual ~FEpicRtcWebsocket() = default;
 
-		/* Begin EpicRtcWebsocketInterface */
+		// Begin EpicRtcWebsocketInterface
 		virtual EpicRtcBool Connect(EpicRtcStringView Url, EpicRtcWebsocketObserverInterface* Observer) override;
 		virtual void		Disconnect(const EpicRtcStringView Reason) override;
 		virtual void		Send(EpicRtcStringView Message) override;
-		/* End EpicRtcWebsocketInterface */
+		// End EpicRtcWebsocketInterface
 
-		/* Begin FTickableGameObject */
+		// Begin FTickableGameObject
 		virtual void		Tick(float DeltaTime) override;
 		virtual bool		IsTickableInEditor() const override { return true; }
 		virtual bool		IsTickableWhenPaused() const override { return true; }
 		FORCEINLINE TStatId GetStatId() const { RETURN_QUICK_DECLARE_CYCLE_STAT(EpicRtcWebSocket, STATGROUP_Tickables); }
-		/* End FTickableGameObject */
+		// End FTickableGameObject
 
 	private:
 		void OnConnected();
@@ -58,11 +57,11 @@ namespace UE::PixelStreaming2
 		bool											bCloseRequested = false;
 
 	public:
-		/* Begin EpicRtcRefCountInterface */
+		// Begin EpicRtcRefCountInterface
 		virtual uint32_t AddRef() override final { return TRefCountingMixin<FEpicRtcWebsocket>::AddRef(); }
 		virtual uint32_t Release() override final { return TRefCountingMixin<FEpicRtcWebsocket>::Release(); }
 		virtual uint32_t Count() const override final { return TRefCountingMixin<FEpicRtcWebsocket>::GetRefCount(); }
-		/* End EpicRtcRefCountInterface */
+		// End EpicRtcRefCountInterface
 	};
 
 } // namespace UE::PixelStreaming2
\ No newline at end of file
diff --git a/Engine/Plugins/Media/PixelStreaming2/Source/PixelStreaming2/Internal/EpicRtcWebsocketFactory.h b/Engine/Plugins/Media/PixelStreaming2/Source/PixelStreaming2/Internal/EpicRtcWebsocketFactory.h
--- a/Engine/Plugins/Media/PixelStreaming2/Source/PixelStreaming2/Internal/EpicRtcWebsocketFactory.h
+++ b/Engine/Plugins/Media/PixelStreaming2/Source/PixelStreaming2/Internal/EpicRtcWebsocketFactory.h
@@ -3,26 +3,29 @@
 #pragma once
 
 #include "EpicRtcWebsocket.h"
+
 #include "epic_rtc/plugins/signalling/websocket.h"
 #include "epic_rtc/plugins/signalling/websocket_factory.h"
 
 namespace UE::PixelStreaming2
 {
-
 	class PIXELSTREAMING2_API FEpicRtcWebsocketFactory : public EpicRtcWebsocketFactoryInterface, public TRefCountingMixin<FEpicRtcWebsocketFactory>
 	{
 	public:
 		FEpicRtcWebsocketFactory(bool bInSendKeepAlive = true)
-			: bSendKeepAlive(bInSendKeepAlive) {};
+			: bSendKeepAlive(bInSendKeepAlive)
+		{
+		}
 		virtual ~FEpicRtcWebsocketFactory() = default;
+
 		virtual EpicRtcErrorCode CreateWebsocket(EpicRtcWebsocketInterface** OutWebsocket) override;
 
 	public:
-		/* Begin EpicRtcRefCountInterface */
+		// Begin EpicRtcRefCountInterface
 		virtual uint32_t AddRef() override final { return TRefCountingMixin<FEpicRtcWebsocketFactory>::AddRef(); }
 		virtual uint32_t Release() override final { return TRefCountingMixin<FEpicRtcWebsocketFactory>::Release(); }
 		virtual uint32_t Count() const override final { return TRefCountingMixin<FEpicRtcWebsocketFactory>::GetRefCount(); }
-		/* End EpicRtcRefCountInterface */
+		// End EpicRtcRefCountInterface
 
 	private:
 		bool bSendKeepAlive;
diff --git a/Engine/Plugins/Media/PixelStreaming2/Source/PixelStreaming2/Internal/SoundGenerator.h b/Engine/Plugins/Media/PixelStreaming2/Source/PixelStreaming2/Internal/SoundGenerator.h
--- a/Engine/Plugins/Media/PixelStreaming2/Source/PixelStreaming2/Internal/SoundGenerator.h
+++ b/Engine/Plugins/Media/PixelStreaming2/Source/PixelStreaming2/Internal/SoundGenerator.h
@@ -2,14 +2,14 @@
 
 #pragma once
 
-#include "CoreMinimal.h"
-#include "HAL/ThreadSafeBool.h"
+#include "CoreTypes.h"
 #include "Sound/SoundGenerator.h"
 
+#include <atomic>
+
 namespace UE::PixelStreaming2
 {
-
-	/*
+	/**
 	 * An `ISoundGenerator` implementation to pump some audio from EpicRtc into this synth component
 	 */
 	class PIXELSTREAMING2_API FSoundGenerator : public ::ISoundGenerator
@@ -46,8 +46,8 @@ namespace UE::PixelStreaming2
 		FCriticalSection		  CriticalSection;
 
 	public:
-		FThreadSafeBool bGeneratingAudio = false;
-		FThreadSafeBool bShouldGenerateAudio = false;
+		std::atomic<bool> bGeneratingAudio = false;
+		std::atomic<bool> bShouldGenerateAudio = false;
 	};
 
 } // namespace UE::PixelStreaming2
diff --git a/Engine/Plugins/Media/PixelStreaming2/Source/PixelStreaming2/Internal/ThreadSafeMap.h b/Engine/Plugins/Media/PixelStreaming2/Source/PixelStreaming2/Internal/ThreadSafeMap.h
--- a/Engine/Plugins/Media/PixelStreaming2/Source/PixelStreaming2/Internal/ThreadSafeMap.h
+++ b/Engine/Plugins/Media/PixelStreaming2/Source/PixelStreaming2/Internal/ThreadSafeMap.h
@@ -4,12 +4,10 @@
 
 #include "Containers/ContainerAllocationPolicies.h"
 #include "Containers/Map.h"
-#include "HAL/ThreadSafeBool.h"
 #include "Misc/ScopeLock.h"
 
 namespace UE::PixelStreaming2
 {
-
 	/** A TThreadSafeMapBase specialization that only allows a single value associated with each key whilst being thread safe.*/
 	template <typename InKeyType, typename InValueType, typename SetAllocator = FDefaultSetAllocator, typename KeyFuncs = TDefaultMapHashableKeyFuncs<InKeyType, InValueType, false>>
 	class TThreadSafeMap : protected TSortableMapBase<InKeyType, InValueType, SetAllocator, KeyFuncs>
diff --git a/Engine/Plugins/Media/PixelStreaming2/Source/PixelStreaming2/Internal/VideoCapturer.h b/Engine/Plugins/Media/PixelStreaming2/Source/PixelStreaming2/Internal/VideoCapturer.h
--- a/Engine/Plugins/Media/PixelStreaming2/Source/PixelStreaming2/Internal/VideoCapturer.h
+++ b/Engine/Plugins/Media/PixelStreaming2/Source/PixelStreaming2/Internal/VideoCapturer.h
@@ -2,9 +2,8 @@
 
 #pragma once
 
-#include "Delegates/IDelegateInstance.h"
+#include "Delegates/Delegate.h"
 #include "IPixelCaptureCapturerSource.h"
-#include "PixelCaptureBufferFormat.h"
 #include "PixelCaptureCapturerMultiFormat.h"
 #include "VideoProducer.h"
 
@@ -12,7 +11,6 @@
 
 namespace UE::PixelStreaming2
 {
-
 	/**
 	 * The start of PixelCapture pipeline. Frames enter the system when `OnFrameCaptured` is called.
 	 * This class creates the underlying PixelCapture `FPixelCaptureCapturer` that handles frame capture when `RequestFormat` is called.
@@ -23,9 +21,9 @@ namespace UE::PixelStreaming2
 		static TSharedPtr<FVideoCapturer> Create(TSharedPtr<FVideoProducer> VideoProducer = nullptr);
 		virtual ~FVideoCapturer() = default;
 
-		/** Begin IPixelCaptureCapturerSource Interface */
+		// Begin IPixelCaptureCapturerSource Interface
 		virtual TSharedPtr<FPixelCaptureCapturer> CreateCapturer(int32 FinalFormat, float FinalScale) override;
-		/** End IPixelCaptureCapturerSource Interface */
+		// End IPixelCaptureCapturerSource Interface
 
 		bool IsReady() const { return bReady; }
 
diff --git a/Engine/Plugins/Media/PixelStreaming2/Source/PixelStreaming2/Internal/VideoProducer.h b/Engine/Plugins/Media/PixelStreaming2/Source/PixelStreaming2/Internal/VideoProducer.h
--- a/Engine/Plugins/Media/PixelStreaming2/Source/PixelStreaming2/Internal/VideoProducer.h
+++ b/Engine/Plugins/Media/PixelStreaming2/Source/PixelStreaming2/Internal/VideoProducer.h
@@ -2,11 +2,12 @@
 
 #pragma once
 
+#include "Delegates/DelegateCombinations.h"
 #include "IPixelStreaming2VideoProducer.h"
+#include "Templates/SharedPointer.h"
 
 namespace UE::PixelStreaming2
 {
-
 	class PIXELSTREAMING2_API FVideoProducer : public IPixelStreaming2VideoProducer
 	{
 	public:
diff --git a/Engine/Plugins/Media/PixelStreaming2/Source/PixelStreaming2/Internal/VideoProducerBackBuffer.h b/Engine/Plugins/Media/PixelStreaming2/Source/PixelStreaming2/Internal/VideoProducerBackBuffer.h
--- a/Engine/Plugins/Media/PixelStreaming2/Source/PixelStreaming2/Internal/VideoProducerBackBuffer.h
+++ b/Engine/Plugins/Media/PixelStreaming2/Source/PixelStreaming2/Internal/VideoProducerBackBuffer.h
@@ -9,8 +9,7 @@
 
 namespace UE::PixelStreaming2
 {
-
-	/*
+	/**
 	 * Use this if you want to send the contents of the UE backbuffer.
 	 */
 	class PIXELSTREAMING2_API FVideoProducerBackBuffer : public FVideoProducer
diff --git a/Engine/Plugins/Media/PixelStreaming2/Source/PixelStreaming2/Internal/VideoProducerMediaCapture.h b/Engine/Plugins/Media/PixelStreaming2/Source/PixelStreaming2/Internal/VideoProducerMediaCapture.h
--- a/Engine/Plugins/Media/PixelStreaming2/Source/PixelStreaming2/Internal/VideoProducerMediaCapture.h
+++ b/Engine/Plugins/Media/PixelStreaming2/Source/PixelStreaming2/Internal/VideoProducerMediaCapture.h
@@ -5,10 +5,10 @@
 #include "Delegates/IDelegateInstance.h"
 #include "Misc/Optional.h"
 #include "MediaOutput.h"
-
 #include "PixelStreaming2MediaIOCapture.h"
 #include "VideoProducer.h"
 
+
 #include "VideoProducerMediaCapture.generated.h"
 
 UCLASS(BlueprintType)
@@ -29,7 +29,7 @@ private:
 namespace UE::PixelStreaming2
 {
 
-	/*
+	/**
 	 * Use this if you want to send media capture frames as video input.
 	 */
 	class PIXELSTREAMING2_API FVideoProducerMediaCapture : public FVideoProducer, public TSharedFromThis<FVideoProducerMediaCapture>
diff --git a/Engine/Plugins/Media/PixelStreaming2/Source/PixelStreaming2/Internal/VideoProducerPIEViewport.h b/Engine/Plugins/Media/PixelStreaming2/Source/PixelStreaming2/Internal/VideoProducerPIEViewport.h
--- a/Engine/Plugins/Media/PixelStreaming2/Source/PixelStreaming2/Internal/VideoProducerPIEViewport.h
+++ b/Engine/Plugins/Media/PixelStreaming2/Source/PixelStreaming2/Internal/VideoProducerPIEViewport.h
@@ -2,14 +2,14 @@
 
 #pragma once
 
+#include "Delegates/IDelegateInstance.h"
 #include "VideoProducer.h"
 
 class FViewport;
 
 namespace UE::PixelStreaming2
 {
-
-	/*
+	/**
 	 * An extension of the back buffer input that can handle PIE sessions. Primarily to be used in blueprints
 	 */
 	class PIXELSTREAMING2_API FVideoProducerPIEViewport : public FVideoProducer
diff --git a/Engine/Plugins/Media/PixelStreaming2/Source/PixelStreaming2/Internal/VideoProducerRenderTarget.h b/Engine/Plugins/Media/PixelStreaming2/Source/PixelStreaming2/Internal/VideoProducerRenderTarget.h
--- a/Engine/Plugins/Media/PixelStreaming2/Source/PixelStreaming2/Internal/VideoProducerRenderTarget.h
+++ b/Engine/Plugins/Media/PixelStreaming2/Source/PixelStreaming2/Internal/VideoProducerRenderTarget.h
@@ -4,15 +4,12 @@
 
 #include "Delegates/IDelegateInstance.h"
 #include "VideoProducer.h"
-#include "RHI.h"
-#include "Widgets/SWindow.h"
 
 class UTextureRenderTarget2D;
 
 namespace UE::PixelStreaming2
 {
-
-	/*
+	/**
 	 * Use this if you want to send the contents of a render target.
 	 */
 	class PIXELSTREAMING2_API FVideoProducerRenderTarget : public FVideoProducer
diff --git a/Engine/Plugins/Media/PixelStreaming2/Source/PixelStreaming2/Internal/VideoSink.h b/Engine/Plugins/Media/PixelStreaming2/Source/PixelStreaming2/Internal/VideoSink.h
--- a/Engine/Plugins/Media/PixelStreaming2/Source/PixelStreaming2/Internal/VideoSink.h
+++ b/Engine/Plugins/Media/PixelStreaming2/Source/PixelStreaming2/Internal/VideoSink.h
@@ -2,6 +2,9 @@
 
 #pragma once
 
+#include "Containers/Set.h"
+#include "CoreTypes.h"
+#include "HAL/CriticalSection.h"
 #include "IPixelStreaming2VideoSink.h"
 #include "RendererInterface.h"
 
@@ -9,7 +12,6 @@
 
 namespace UE::PixelStreaming2
 {
-
 	class PIXELSTREAMING2_API FVideoSink : public IPixelStreaming2VideoSink
 	{
 	public:
@@ -25,7 +27,7 @@ namespace UE::PixelStreaming2
 		void SetMuted(bool bIsMuted);
 
 	protected:
-		FCriticalSection						   VideoConsumersCS;
+		FCriticalSection					 VideoConsumersCS;
 		TSet<IPixelStreaming2VideoConsumer*> VideoConsumers;
 
 		bool bIsMuted = false;
diff --git a/Engine/Plugins/Media/PixelStreaming2/Source/PixelStreaming2/Private/Blueprints/PixelStreaming2AudioComponent.cpp b/Engine/Plugins/Media/PixelStreaming2/Source/PixelStreaming2/Private/Blueprints/PixelStreaming2AudioComponent.cpp
--- a/Engine/Plugins/Media/PixelStreaming2/Source/PixelStreaming2/Private/Blueprints/PixelStreaming2AudioComponent.cpp
+++ b/Engine/Plugins/Media/PixelStreaming2/Source/PixelStreaming2/Private/Blueprints/PixelStreaming2AudioComponent.cpp
@@ -9,7 +9,7 @@
 #include "SampleBuffer.h"
 #include "SoundGenerator.h"
 
-/*
+/**
  * Component that recieves audio from a remote webrtc connection and outputs it into UE using a "synth component".
  */
 UPixelStreaming2AudioComponent::UPixelStreaming2AudioComponent(const FObjectInitializer& ObjectInitializer)
diff --git a/Engine/Plugins/Media/PixelStreaming2/Source/PixelStreaming2/Private/Blueprints/PixelStreaming2Peer.cpp b/Engine/Plugins/Media/PixelStreaming2/Source/PixelStreaming2/Private/Blueprints/PixelStreaming2Peer.cpp
--- a/Engine/Plugins/Media/PixelStreaming2/Source/PixelStreaming2/Private/Blueprints/PixelStreaming2Peer.cpp
+++ b/Engine/Plugins/Media/PixelStreaming2/Source/PixelStreaming2/Private/Blueprints/PixelStreaming2Peer.cpp
@@ -8,7 +8,7 @@
 #include "PixelStreaming2Module.h"
 #include "SampleBuffer.h"
 #include "SoundGenerator.h"
-#include "ToStringExtensions.h"
+#include "UtilsString.h"
 
 uint32_t UPixelStreaming2Peer::PlayerId = 0;
 
diff --git a/Engine/Plugins/Media/PixelStreaming2/Source/PixelStreaming2/Private/Blueprints/PixelStreaming2Peer.h b/Engine/Plugins/Media/PixelStreaming2/Source/PixelStreaming2/Private/Blueprints/PixelStreaming2Peer.h
--- a/Engine/Plugins/Media/PixelStreaming2/Source/PixelStreaming2/Private/Blueprints/PixelStreaming2Peer.h
+++ b/Engine/Plugins/Media/PixelStreaming2/Source/PixelStreaming2/Private/Blueprints/PixelStreaming2Peer.h
@@ -92,13 +92,13 @@ private:
 	bool Disconnect(const FString& OptionalReason);
 
 public:
-	/* Begin IPixelStreaming2AudioConsumer Callbacks */
+	// Begin IPixelStreaming2AudioConsumer Callbacks
 	virtual void ConsumeRawPCM(const int16_t* AudioData, int InSampleRate, size_t NChannels, size_t NFrames) override;
 	virtual void OnConsumerAdded() override;
 	virtual void OnConsumerRemoved() override;
-	/* End IPixelStreaming2AudioConsumer Callbacks */
+	// End IPixelStreaming2AudioConsumer Callbacks
 
-	/* Begin FEpicRtcManager Callbacks */
+	// Begin FEpicRtcManager Callbacks
 	void OnSessionStateUpdate(const EpicRtcSessionState StateUpdate);
 	void OnSessionErrorUpdate(const EpicRtcErrorCode ErrorUpdate);
 	void OnSessionRoomsAvailableUpdate(EpicRtcStringArrayInterface* RoomsList);
@@ -126,7 +126,7 @@ public:
 	void OnDataTrackRemoved(EpicRtcDataTrackInterface* DataTrack);
 	void OnDataTrackState(EpicRtcDataTrackInterface* DataTrack, const EpicRtcTrackState State);
 	void OnDataTrackMessage(EpicRtcDataTrackInterface* DataTrack);
-	/* End FEpicRtcManager Callbacks */
+	// End FEpicRtcManager Callbacks
 
 	void OnStatsReady(const FString& PeerId, const EpicRtcConnectionStats& ConnectionStats);
 
diff --git a/Engine/Plugins/Media/PixelStreaming2/Source/PixelStreaming2/Private/Blueprints/PixelStreaming2VideoComponent.cpp b/Engine/Plugins/Media/PixelStreaming2/Source/PixelStreaming2/Private/Blueprints/PixelStreaming2VideoComponent.cpp
--- a/Engine/Plugins/Media/PixelStreaming2/Source/PixelStreaming2/Private/Blueprints/PixelStreaming2VideoComponent.cpp
+++ b/Engine/Plugins/Media/PixelStreaming2/Source/PixelStreaming2/Private/Blueprints/PixelStreaming2VideoComponent.cpp
@@ -7,7 +7,7 @@
 #include "IPixelStreaming2Streamer.h"
 #include "Logging.h"
 
-/*
+/**
  * Component that recieves video from a remote webrtc connection and outputs it into UE using a "synth component".
  */
 UPixelStreaming2VideoComponent::UPixelStreaming2VideoComponent(const FObjectInitializer& ObjectInitializer)
diff --git a/Engine/Plugins/Media/PixelStreaming2/Source/PixelStreaming2/Private/Blueprints/PixelStreaming2VideoProducer.cpp b/Engine/Plugins/Media/PixelStreaming2/Source/PixelStreaming2/Private/Blueprints/PixelStreaming2VideoProducer.cpp
--- a/Engine/Plugins/Media/PixelStreaming2/Source/PixelStreaming2/Private/Blueprints/PixelStreaming2VideoProducer.cpp
+++ b/Engine/Plugins/Media/PixelStreaming2/Source/PixelStreaming2/Private/Blueprints/PixelStreaming2VideoProducer.cpp
@@ -7,7 +7,7 @@
 #include "VideoProducerPIEViewport.h"
 #include "VideoProducerRenderTarget.h"
 
-/*
+/**
  * ---------- UPixelStreaming2VideoProducer -------------------
  */
 UPixelStreaming2VideoProducer::UPixelStreaming2VideoProducer(const FObjectInitializer& ObjectInitializer)
@@ -15,7 +15,7 @@ UPixelStreaming2VideoProducer::UPixelStreaming2VideoProducer(const FObjectInitia
 {
 }
 
-/*
+/**
  * ---------- UPixelStreaming2VideoProducerBackBuffer -------------------
  */
 UPixelStreaming2VideoProducerBackBuffer::UPixelStreaming2VideoProducerBackBuffer(const FObjectInitializer& ObjectInitializer)
@@ -43,7 +43,7 @@ TSharedPtr<IPixelStreaming2VideoProducer> UPixelStreaming2VideoProducerBackBuffe
 	return VideoProducer;
 }
 
-/*
+/**
  * ---------- UPixelStreaming2VideoProducerMediaCapture -------------------
  */
 UPixelStreaming2VideoProducerMediaCapture::UPixelStreaming2VideoProducerMediaCapture(const FObjectInitializer& ObjectInitializer)
@@ -71,7 +71,7 @@ TSharedPtr<IPixelStreaming2VideoProducer> UPixelStreaming2VideoProducerMediaCapt
 	return VideoProducer;
 }
 
-/*
+/**
  * ---------- UPixelStreaming2VideoProducerRenderTarget -------------------
  */
 UPixelStreaming2VideoProducerRenderTarget::UPixelStreaming2VideoProducerRenderTarget(const FObjectInitializer& ObjectInitializer)
@@ -87,4 +87,3 @@ TSharedPtr<IPixelStreaming2VideoProducer> UPixelStreaming2VideoProducerRenderTar
 	}
 	return VideoProducer;
 }
-
diff --git a/Engine/Plugins/Media/PixelStreaming2/Source/PixelStreaming2/Private/EpicRtcAudioMixingCapturer.h b/Engine/Plugins/Media/PixelStreaming2/Source/PixelStreaming2/Private/EpicRtcAudioMixingCapturer.h
--- a/Engine/Plugins/Media/PixelStreaming2/Source/PixelStreaming2/Private/EpicRtcAudioMixingCapturer.h
+++ b/Engine/Plugins/Media/PixelStreaming2/Source/PixelStreaming2/Private/EpicRtcAudioMixingCapturer.h
@@ -85,7 +85,7 @@ namespace UE::PixelStreaming2
 
 		virtual ~FEpicRtcMixAudioTask() = default;
 
-		/* Begin FEpicRtcMixAudioTask */
+		// Begin FEpicRtcMixAudioTask
 		virtual void Tick(float DeltaMs) override
 		{
 			if (!Mixer)
@@ -119,7 +119,7 @@ namespace UE::PixelStreaming2
 			static FString TaskName = TEXT("EpicRtcMixAudioTask");
 			return TaskName;
 		}
-		/* End FEpicRtcMixAudioTask */
+		// End FEpicRtcMixAudioTask
 
 	private:
 		bool								  bIsRunning;
diff --git a/Engine/Plugins/Media/PixelStreaming2/Source/PixelStreaming2/Private/EpicRtcAudioPatchMixer.cpp b/Engine/Plugins/Media/PixelStreaming2/Source/PixelStreaming2/Private/EpicRtcAudioPatchMixer.cpp
--- a/Engine/Plugins/Media/PixelStreaming2/Source/PixelStreaming2/Private/EpicRtcAudioPatchMixer.cpp
+++ b/Engine/Plugins/Media/PixelStreaming2/Source/PixelStreaming2/Private/EpicRtcAudioPatchMixer.cpp
@@ -25,4 +25,4 @@ namespace UE::PixelStreaming2
 	{
 		return SampleRate;
 	}
-}
\ No newline at end of file
+} // namespace UE::PixelStreaming2
\ No newline at end of file
diff --git a/Engine/Plugins/Media/PixelStreaming2/Source/PixelStreaming2/Private/EpicRtcAudioPatchMixer.h b/Engine/Plugins/Media/PixelStreaming2/Source/PixelStreaming2/Private/EpicRtcAudioPatchMixer.h
--- a/Engine/Plugins/Media/PixelStreaming2/Source/PixelStreaming2/Private/EpicRtcAudioPatchMixer.h
+++ b/Engine/Plugins/Media/PixelStreaming2/Source/PixelStreaming2/Private/EpicRtcAudioPatchMixer.h
@@ -6,18 +6,18 @@
 
 namespace UE::PixelStreaming2
 {
-    class FEpicRtcAudioPatchMixer : public Audio::FPatchMixer
-    {
-    public:
+	class FEpicRtcAudioPatchMixer : public Audio::FPatchMixer
+	{
+	public:
 		FEpicRtcAudioPatchMixer(uint8 NumChannels, uint32 SampleRate, float SampleSizeSeconds);
 
 		uint32 GetMaxBufferSize() const;
-		uint8 GetNumChannels() const;
+		uint8  GetNumChannels() const;
 		uint32 GetSampleRate() const;
-        
+
 	private:
-		uint8 NumChannels;
+		uint8  NumChannels;
 		uint32 SampleRate;
-		float SampleSizeSeconds;
-    };
-}
\ No newline at end of file
+		float  SampleSizeSeconds;
+	};
+} // namespace UE::PixelStreaming2
\ No newline at end of file
diff --git a/Engine/Plugins/Media/PixelStreaming2/Source/PixelStreaming2/Private/EpicRtcAudioProducer.h b/Engine/Plugins/Media/PixelStreaming2/Source/PixelStreaming2/Private/EpicRtcAudioProducer.h
--- a/Engine/Plugins/Media/PixelStreaming2/Source/PixelStreaming2/Private/EpicRtcAudioProducer.h
+++ b/Engine/Plugins/Media/PixelStreaming2/Source/PixelStreaming2/Private/EpicRtcAudioProducer.h
@@ -13,24 +13,24 @@ namespace UE::PixelStreaming2
 	/**
 	 * An audio input capable of listening to UE submix's as well as receiving user audio via the PushAudio method.
 	 * Any received audio will be passed into the Parent's PushAudio method
-	*/
-    class FEpicRtcAudioProducer : public ISubmixBufferListener, public IPixelStreaming2AudioProducer
+	 */
+	class FEpicRtcAudioProducer : public ISubmixBufferListener, public IPixelStreaming2AudioProducer
 	{
-    public:
+	public:
 		static TSharedPtr<FEpicRtcAudioProducer> Create(Audio::FDeviceId AudioDeviceId, TSharedPtr<FEpicRtcPatchInputProxy> PatchInput);
 		static TSharedPtr<FEpicRtcAudioProducer> Create(TSharedPtr<FEpicRtcPatchInputProxy> PatchInput);
-        virtual ~FEpicRtcAudioProducer() = default;
+		virtual ~FEpicRtcAudioProducer() = default;
 
-		/* For users to manually push non-submix audio into EpicRtc */
-    	virtual void PushAudio(const float* AudioData, int32 NumSamples, int32 NumChannels, int32 SampleRate) override;
+		// For users to manually push non-submix audio into EpicRtc
+		virtual void PushAudio(const float* AudioData, int32 NumSamples, int32 NumChannels, int32 SampleRate) override;
 
 		// ISubmixBufferListener interface
 		virtual void OnNewSubmixBuffer(const USoundSubmix* OwningSubmix, float* AudioData, int32 NumSamples, int32 NumChannels, const int32 SampleRate, double AudioClock) override;
 
-    protected:
+	protected:
 		FEpicRtcAudioProducer(TSharedPtr<FEpicRtcPatchInputProxy> PatchInput);
 
-    private:
-		TSharedPtr<FEpicRtcPatchInputProxy> PatchInput;	
-	}; 
+	private:
+		TSharedPtr<FEpicRtcPatchInputProxy> PatchInput;
+	};
 } // namespace UE::PixelStreaming2
\ No newline at end of file
diff --git a/Engine/Plugins/Media/PixelStreaming2/Source/PixelStreaming2/Private/EpicRtcAudioSink.h b/Engine/Plugins/Media/PixelStreaming2/Source/PixelStreaming2/Private/EpicRtcAudioSink.h
--- a/Engine/Plugins/Media/PixelStreaming2/Source/PixelStreaming2/Private/EpicRtcAudioSink.h
+++ b/Engine/Plugins/Media/PixelStreaming2/Source/PixelStreaming2/Private/EpicRtcAudioSink.h
@@ -2,10 +2,8 @@
 
 #pragma once
 
-#include "Containers/Set.h"
-#include "EpicRtcAudioTrack.h"
-#include "IPixelStreaming2AudioConsumer.h"
 #include "AudioSink.h"
+#include "EpicRtcAudioTrack.h"
 
 namespace UE::PixelStreaming2
 {
diff --git a/Engine/Plugins/Media/PixelStreaming2/Source/PixelStreaming2/Private/EpicRtcAudioSource.cpp b/Engine/Plugins/Media/PixelStreaming2/Source/PixelStreaming2/Private/EpicRtcAudioSource.cpp
--- a/Engine/Plugins/Media/PixelStreaming2/Source/PixelStreaming2/Private/EpicRtcAudioSource.cpp
+++ b/Engine/Plugins/Media/PixelStreaming2/Source/PixelStreaming2/Private/EpicRtcAudioSource.cpp
@@ -3,7 +3,6 @@
 #include "EpicRtcAudioSource.h"
 
 #include "EpicRtcAudioMixingCapturer.h"
-#include "Utils.h"
 #include "PixelStreaming2Module.h"
 
 namespace UE::PixelStreaming2
diff --git a/Engine/Plugins/Media/PixelStreaming2/Source/PixelStreaming2/Private/EpicRtcDataTrack.cpp b/Engine/Plugins/Media/PixelStreaming2/Source/PixelStreaming2/Private/EpicRtcDataTrack.cpp
--- a/Engine/Plugins/Media/PixelStreaming2/Source/PixelStreaming2/Private/EpicRtcDataTrack.cpp
+++ b/Engine/Plugins/Media/PixelStreaming2/Source/PixelStreaming2/Private/EpicRtcDataTrack.cpp
@@ -1,12 +1,12 @@
 // Copyright Epic Games, Inc. All Rights Reserved.
 
 #include "EpicRtcDataTrack.h"
-#include "Utils.h"
-#include "ToStringExtensions.h"
+
+#include "UtilsCommon.h"
+#include "UtilsString.h"
 
 namespace UE::PixelStreaming2
 {
-
 	TSharedPtr<FEpicRtcDataTrack> FEpicRtcDataTrack::Create(TRefCountPtr<EpicRtcDataTrackInterface> InTrack, TWeakPtr<IPixelStreaming2DataProtocol> InDataProtocol)
 	{
 		TSharedPtr<FEpicRtcDataTrack> DataTrack = TSharedPtr<FEpicRtcDataTrack>(new FEpicRtcDataTrack(InTrack, InDataProtocol));
@@ -32,11 +32,13 @@ namespace UE::PixelStreaming2
 			UE_LOG(LogPixelStreaming2, Error, TEXT("Cannot send message when datatrack is null."));
 			return false;
 		}
+
 		if (Track->GetState() != EpicRtcTrackState::Active)
 		{
 			UE_LOG(LogPixelStreaming2, Error, TEXT("Cannot send message when datatrack is not active."));
 			return false;
 		}
+
 		return true;
 	}
 
@@ -56,7 +58,9 @@ namespace UE::PixelStreaming2
 			UE_LOG(LogPixelStreaming2, Error, TEXT("Cannot send message called '%s' as it is not in the data protocol. Try GetTo/FromStreamerProtocol()->Add()"), *MessageType);
 			return false;
 		}
+
 		OutMessageId = Message->GetID();
+
 		return true;
 	}
 
@@ -67,12 +71,15 @@ namespace UE::PixelStreaming2
 			._size = (uint32_t)Buffer.Num(),
 			._binary = true
 		};
+
 		TRefCountPtr<EpicRtcDataTrackInterface> OutgoingTrack = SendTrack ? SendTrack : Track;
+
 		EpicRtcBool SendResult = OutgoingTrack->PushFrame(DataFrame);
 		if (!SendResult)
 		{
 			UE_LOG(LogPixelStreaming2, Error, TEXT("DataTrack PushFrame return false"));
 		}
+
 		return static_cast<bool>(SendResult);
 	}
 
@@ -83,11 +90,12 @@ namespace UE::PixelStreaming2
 
 	bool FEpicRtcDataTrack::SendArbitraryData(const FString& MessageType, const TArray64<uint8>& DataBytes) const
 	{
-		uint8 Type;
 		if (!IsActive())
 		{
 			return false;
 		}
+
+		uint8 Type;
 		if (!GetMessageId(MessageType, Type))
 		{
 			return false;
diff --git a/Engine/Plugins/Media/PixelStreaming2/Source/PixelStreaming2/Private/EpicRtcStatsCollector.cpp b/Engine/Plugins/Media/PixelStreaming2/Source/PixelStreaming2/Private/EpicRtcStatsCollector.cpp
--- a/Engine/Plugins/Media/PixelStreaming2/Source/PixelStreaming2/Private/EpicRtcStatsCollector.cpp
+++ b/Engine/Plugins/Media/PixelStreaming2/Source/PixelStreaming2/Private/EpicRtcStatsCollector.cpp
@@ -2,7 +2,7 @@
 
 #include "EpicRtcStatsCollector.h"
 
-#include "ToStringExtensions.h"
+#include "UtilsString.h"
 
 namespace UE::PixelStreaming2
 {
diff --git a/Engine/Plugins/Media/PixelStreaming2/Source/PixelStreaming2/Private/EpicRtcStatsCollector.h b/Engine/Plugins/Media/PixelStreaming2/Source/PixelStreaming2/Private/EpicRtcStatsCollector.h
--- a/Engine/Plugins/Media/PixelStreaming2/Source/PixelStreaming2/Private/EpicRtcStatsCollector.h
+++ b/Engine/Plugins/Media/PixelStreaming2/Source/PixelStreaming2/Private/EpicRtcStatsCollector.h
@@ -18,6 +18,9 @@ namespace UE::PixelStreaming2
 
 		// Begin EpicRtcStatsCollectorCallbackInterface interface
 		void			 OnStatsDelivered(const EpicRtcStatsReport& InReport) override;
+		// End EpicRtcStatsCollectorCallbackInterface interface
+
+		// Begin EpicRtcRefCountInterface
 		virtual uint32_t AddRef() override final { return TRefCountingMixin<FEpicRtcStatsCollector>::AddRef(); }
 		virtual uint32_t Release() override final { return TRefCountingMixin<FEpicRtcStatsCollector>::Release(); }
 		virtual uint32_t Count() const override final { return TRefCountingMixin<FEpicRtcStatsCollector>::GetRefCount(); }
diff --git a/Engine/Plugins/Media/PixelStreaming2/Source/PixelStreaming2/Private/EpicRtcThread.cpp b/Engine/Plugins/Media/PixelStreaming2/Source/PixelStreaming2/Private/EpicRtcThread.cpp
--- a/Engine/Plugins/Media/PixelStreaming2/Source/PixelStreaming2/Private/EpicRtcThread.cpp
+++ b/Engine/Plugins/Media/PixelStreaming2/Source/PixelStreaming2/Private/EpicRtcThread.cpp
@@ -20,7 +20,7 @@ namespace UE::PixelStreaming2
 	class FEpicRtcRunnable : public FRunnable, public FSingleThreadRunnable
 	{
 	public:
-		/* Begin FRunnable */
+		// Begin FRunnable
 		virtual bool Init() override
 		{
 			return true;
@@ -58,9 +58,9 @@ namespace UE::PixelStreaming2
 		{
 			return this;
 		}
-		/* End FRunnable */
+		// End FRunnable
 
-		/* Begin FSingleThreadRunnable */
+		// Begin FSingleThreadRunnable
 		virtual void Tick() override
 		{
 			FScopeLock TaskLock(&TasksMutex);
@@ -83,7 +83,7 @@ namespace UE::PixelStreaming2
 
 			LastTickCycles = NowCycles;
 		}
-		/* End FSingleThreadRunnable */
+		// End FSingleThreadRunnable
 
 		void StartTicking()
 		{
@@ -175,7 +175,7 @@ namespace UE::PixelStreaming2
 		// This critical section should be locked during entire tick process
 		FCriticalSection TasksMutex;
 
-		/* Use this event to signal when we should wake. */
+		// Use this event to signal when we should wake.
 		FEventRef TaskEvent;
 
 		// Tasks can removed from any thread so this needs to be thread safe
diff --git a/Engine/Plugins/Media/PixelStreaming2/Source/PixelStreaming2/Private/EpicRtcVideoBufferI420.h b/Engine/Plugins/Media/PixelStreaming2/Source/PixelStreaming2/Private/EpicRtcVideoBufferI420.h
--- a/Engine/Plugins/Media/PixelStreaming2/Source/PixelStreaming2/Private/EpicRtcVideoBufferI420.h
+++ b/Engine/Plugins/Media/PixelStreaming2/Source/PixelStreaming2/Private/EpicRtcVideoBufferI420.h
@@ -2,7 +2,6 @@
 
 #pragma once
 
-#include "EpicRtcMemory.h"
 #include "PixelCaptureBufferI420.h"
 
 #include "epic_rtc/core/video/video_buffer.h"
@@ -17,7 +16,7 @@ namespace UE::PixelStreaming2
 		virtual ~FEpicRtcVideoBufferI420() = default;
 
 	public:
-		/* Begin EpicRtcVideoBufferInterface */
+		// Begin EpicRtcVideoBufferInterface
 		virtual void* GetData() override
 		{
 			return Buffer->GetMutableData();
@@ -37,14 +36,14 @@ namespace UE::PixelStreaming2
 		{
 			return Buffer->GetHeight();
 		}
-		/* End EpicRtcVideoBufferInterface */
+		// End EpicRtcVideoBufferInterface
 
 	public:
-		/* Begin EpicRtcRefCountInterface */
+		// Begin EpicRtcRefCountInterface
 		virtual uint32_t AddRef() override final { return TRefCountingMixin<FEpicRtcVideoBufferI420>::AddRef(); }
 		virtual uint32_t Release() override final { return TRefCountingMixin<FEpicRtcVideoBufferI420>::Release(); }
 		virtual uint32_t Count() const override final { return TRefCountingMixin<FEpicRtcVideoBufferI420>::GetRefCount(); }
-		/* End EpicRtcRefCountInterface */
+		// End EpicRtcRefCountInterface
 
 	private:
 		TSharedPtr<FPixelCaptureBufferI420> Buffer;
diff --git a/Engine/Plugins/Media/PixelStreaming2/Source/PixelStreaming2/Private/EpicRtcVideoBufferMultiFormat.h b/Engine/Plugins/Media/PixelStreaming2/Source/PixelStreaming2/Private/EpicRtcVideoBufferMultiFormat.h
--- a/Engine/Plugins/Media/PixelStreaming2/Source/PixelStreaming2/Private/EpicRtcVideoBufferMultiFormat.h
+++ b/Engine/Plugins/Media/PixelStreaming2/Source/PixelStreaming2/Private/EpicRtcVideoBufferMultiFormat.h
@@ -2,7 +2,6 @@
 
 #pragma once
 
-#include "EpicRtcMemory.h"
 #include "PixelCaptureCapturerMultiFormat.h"
 #include "Video/Resources/VideoResourceRHI.h"
 
@@ -22,7 +21,7 @@ namespace UE::PixelStreaming2
 		virtual ~FEpicRtcVideoBufferMultiFormatBase() = default;
 
 	public:
-		/* Begin EpicRtcVideoBufferInterface */
+		// Begin EpicRtcVideoBufferInterface
 		virtual void* GetData() override
 		{
 			unimplemented();
@@ -33,14 +32,14 @@ namespace UE::PixelStreaming2
 		{
 			return EpicRtcPixelFormat::Native;
 		}
-		/* End EpicRtcVideoBufferInterface */
+		// End EpicRtcVideoBufferInterface
 
 	public:
-		/* Begin EpicRtcRefCountInterface */
+		// Begin EpicRtcRefCountInterface
 		virtual uint32_t AddRef() override final { return TRefCountingMixin<FEpicRtcVideoBufferMultiFormatBase>::AddRef(); }
 		virtual uint32_t Release() override final { return TRefCountingMixin<FEpicRtcVideoBufferMultiFormatBase>::Release(); }
 		virtual uint32_t Count() const override final { return TRefCountingMixin<FEpicRtcVideoBufferMultiFormatBase>::GetRefCount(); }
-		/* End EpicRtcRefCountInterface */
+		// End EpicRtcRefCountInterface
 
 	protected:
 		TSharedPtr<FPixelCaptureCapturerMultiFormat> FrameCapturer;
diff --git a/Engine/Plugins/Media/PixelStreaming2/Source/PixelStreaming2/Private/EpicRtcVideoBufferRHI.h b/Engine/Plugins/Media/PixelStreaming2/Source/PixelStreaming2/Private/EpicRtcVideoBufferRHI.h
--- a/Engine/Plugins/Media/PixelStreaming2/Source/PixelStreaming2/Private/EpicRtcVideoBufferRHI.h
+++ b/Engine/Plugins/Media/PixelStreaming2/Source/PixelStreaming2/Private/EpicRtcVideoBufferRHI.h
@@ -2,7 +2,6 @@
 
 #pragma once
 
-#include "EpicRtcMemory.h"
 #include "Video/Resources/VideoResourceRHI.h"
 
 #include "epic_rtc/core/video/video_buffer.h"
@@ -13,11 +12,14 @@ namespace UE::PixelStreaming2
 	{
 	public:
 		FEpicRtcVideoBufferRHI(TSharedPtr<FVideoResourceRHI> VideoResourceRHI)
-			: VideoResourceRHI(VideoResourceRHI) {}
+			: VideoResourceRHI(VideoResourceRHI) 
+		{
+		}
+
 		virtual ~FEpicRtcVideoBufferRHI() = default;
 
 	public:
-		/* Begin EpicRtcVideoBufferInterface */
+		// Begin EpicRtcVideoBufferInterface
 		virtual void* GetData() override
 		{
 			unimplemented();
@@ -38,16 +40,16 @@ namespace UE::PixelStreaming2
 		{
 			return VideoResourceRHI->GetDescriptor().Height;
 		}
-		/* End EpicRtcVideoBufferInterface */
+		// End EpicRtcVideoBufferInterface
 
 		TSharedPtr<FVideoResourceRHI> GetVideoResource() { return VideoResourceRHI; }
 
 	public:
-		/* Begin EpicRtcRefCountInterface */
+		// Begin EpicRtcRefCountInterface
 		virtual uint32_t AddRef() override final { return TRefCountingMixin<FEpicRtcVideoBufferRHI>::AddRef(); }
 		virtual uint32_t Release() override final { return TRefCountingMixin<FEpicRtcVideoBufferRHI>::Release(); }
 		virtual uint32_t Count() const override final { return TRefCountingMixin<FEpicRtcVideoBufferRHI>::GetRefCount(); }
-		/* End EpicRtcRefCountInterface */
+		// End EpicRtcRefCountInterface
 
 	private:
 		TSharedPtr<FVideoResourceRHI> VideoResourceRHI;
diff --git a/Engine/Plugins/Media/PixelStreaming2/Source/PixelStreaming2/Private/EpicRtcVideoDecoder.cpp b/Engine/Plugins/Media/PixelStreaming2/Source/PixelStreaming2/Private/EpicRtcVideoDecoder.cpp
--- a/Engine/Plugins/Media/PixelStreaming2/Source/PixelStreaming2/Private/EpicRtcVideoDecoder.cpp
+++ b/Engine/Plugins/Media/PixelStreaming2/Source/PixelStreaming2/Private/EpicRtcVideoDecoder.cpp
@@ -7,8 +7,8 @@
 #include "EpicRtcVideoBufferRHI.h"
 #include "Logging.h"
 #include "PixelStreaming2Trace.h"
-#include "ToStringExtensions.h"
-#include "VideoUtils.h"
+#include "UtilsString.h"
+#include "UtilsVideo.h"
 
 #include "epic_rtc/core/video/video_codec_info.h"
 
diff --git a/Engine/Plugins/Media/PixelStreaming2/Source/PixelStreaming2/Private/EpicRtcVideoDecoder.h b/Engine/Plugins/Media/PixelStreaming2/Source/PixelStreaming2/Private/EpicRtcVideoDecoder.h
--- a/Engine/Plugins/Media/PixelStreaming2/Source/PixelStreaming2/Private/EpicRtcVideoDecoder.h
+++ b/Engine/Plugins/Media/PixelStreaming2/Source/PixelStreaming2/Private/EpicRtcVideoDecoder.h
@@ -2,7 +2,6 @@
 
 #pragma once
 
-#include "EpicRtcMemory.h"
 #include "Video/Decoders/Configs/VideoDecoderConfigAV1.h"
 #include "Video/Decoders/Configs/VideoDecoderConfigH264.h"
 #include "Video/Decoders/Configs/VideoDecoderConfigVP8.h"
@@ -23,14 +22,14 @@ namespace UE::PixelStreaming2
 	public:
 		TEpicRtcVideoDecoder(EpicRtcVideoCodecInfoInterface* CodecInfo);
 
-		/* Begin EpicRtcVideoDecoderInterface */
+		// Begin EpicRtcVideoDecoderInterface
 		[[nodiscard]] virtual EpicRtcStringView GetName() const override;
 		virtual EpicRtcVideoDecoderConfig		GetConfig() const override;
 		virtual EpicRtcMediaResult				SetConfig(const EpicRtcVideoDecoderConfig& VideoDecoderConfig) override;
 		virtual EpicRtcMediaResult				Decode(const EpicRtcEncodedVideoFrame& Frame) override;
 		virtual void							RegisterCallback(EpicRtcVideoDecoderCallbackInterface* Callback) override;
 		virtual void							Reset() override;
-		/* End EpicRtcVideoDecoderInterface */
+		// End EpicRtcVideoDecoderInterface
 	private:
 		TSharedPtr<TVideoDecoder<TVideoResource>>		   Decoder;
 		TUniquePtr<FVideoDecoderConfig>					   InitialVideoConfig;
@@ -46,10 +45,10 @@ namespace UE::PixelStreaming2
 		bool LateInitDecoder();
 
 	public:
-		/* Begin EpicRtcRefCountInterface */
+		// Begin EpicRtcRefCountInterface
 		virtual uint32_t AddRef() override final { return TRefCountingMixin<TEpicRtcVideoDecoder>::AddRef(); }
 		virtual uint32_t Release() override final { return TRefCountingMixin<TEpicRtcVideoDecoder>::Release(); }
 		virtual uint32_t Count() const override final { return TRefCountingMixin<TEpicRtcVideoDecoder>::GetRefCount(); }
-		/* End EpicRtcRefCountInterface */
+		// End EpicRtcRefCountInterface
 	};
 } // namespace UE::PixelStreaming2
\ No newline at end of file
diff --git a/Engine/Plugins/Media/PixelStreaming2/Source/PixelStreaming2/Private/EpicRtcVideoDecoderInitializer.cpp b/Engine/Plugins/Media/PixelStreaming2/Source/PixelStreaming2/Private/EpicRtcVideoDecoderInitializer.cpp
--- a/Engine/Plugins/Media/PixelStreaming2/Source/PixelStreaming2/Private/EpicRtcVideoDecoderInitializer.cpp
+++ b/Engine/Plugins/Media/PixelStreaming2/Source/PixelStreaming2/Private/EpicRtcVideoDecoderInitializer.cpp
@@ -6,9 +6,9 @@
 #include "CoderUtils.h"
 #include "Logging.h"
 #include "PixelStreaming2PluginSettings.h"
-#include "ToStringExtensions.h"
-#include "Utils.h"
-#include "VideoUtils.h"
+#include "UtilsCommon.h"
+#include "UtilsString.h"
+#include "UtilsVideo.h"
 #include "Video/Decoders/Configs/VideoDecoderConfigAV1.h"
 #include "Video/Decoders/Configs/VideoDecoderConfigH264.h"
 #include "Video/Decoders/Configs/VideoDecoderConfigVP8.h"
@@ -136,14 +136,14 @@ namespace UE::PixelStreaming2
 
 		if (UE::PixelStreaming2::IsDecoderSupported<FVideoDecoderConfigVP8>())
 		{
-			Codecs[EVideoCodec::VP8].Add(TRefCountPtr<EpicRtcVideoCodecInfoInterface>(new FVideoCodecInfo(
+			Codecs[EVideoCodec::VP8].Add(TRefCountPtr<EpicRtcVideoCodecInfoInterface>(new FEpicRtcVideoCodecInfo(
 				EpicRtcVideoCodec::VP8,
 				UE::PixelStreaming2::IsHardwareDecoderSupported<FVideoDecoderConfigVP8>())));
 		}
 
 		if (UE::PixelStreaming2::IsDecoderSupported<FVideoDecoderConfigVP9>())
 		{
-			Codecs[EVideoCodec::VP9].Add(TRefCountPtr<EpicRtcVideoCodecInfoInterface>(new FVideoCodecInfo(
+			Codecs[EVideoCodec::VP9].Add(TRefCountPtr<EpicRtcVideoCodecInfoInterface>(new FEpicRtcVideoCodecInfo(
 				EpicRtcVideoCodec::VP9,
 				UE::PixelStreaming2::IsHardwareDecoderSupported<FVideoDecoderConfigVP9>())));
 		}
@@ -151,11 +151,11 @@ namespace UE::PixelStreaming2
 		if (UE::PixelStreaming2::IsDecoderSupported<FVideoDecoderConfigH264>())
 		{
 			using namespace UE::AVCodecCore::H264;
-			Codecs[EVideoCodec::H264].Add(TRefCountPtr<EpicRtcVideoCodecInfoInterface>(new FVideoCodecInfo(
+			Codecs[EVideoCodec::H264].Add(TRefCountPtr<EpicRtcVideoCodecInfoInterface>(new FEpicRtcVideoCodecInfo(
 				EpicRtcVideoCodec::H264,
 				UE::PixelStreaming2::IsHardwareDecoderSupported<FVideoDecoderConfigH264>(),
 				UE::PixelStreaming2::CreateH264Format(EH264Profile::ConstrainedBaseline, EH264Level::Level_3_1))));
-			Codecs[EVideoCodec::H264].Add(TRefCountPtr<EpicRtcVideoCodecInfoInterface>(new FVideoCodecInfo(
+			Codecs[EVideoCodec::H264].Add(TRefCountPtr<EpicRtcVideoCodecInfoInterface>(new FEpicRtcVideoCodecInfo(
 				EpicRtcVideoCodec::H264,
 				UE::PixelStreaming2::IsHardwareDecoderSupported<FVideoDecoderConfigH264>(),
 				UE::PixelStreaming2::CreateH264Format(EH264Profile::Baseline, EH264Level::Level_3_1))));
@@ -163,7 +163,7 @@ namespace UE::PixelStreaming2
 
 		if (UE::PixelStreaming2::IsDecoderSupported<FVideoDecoderConfigAV1>())
 		{
-			Codecs[EVideoCodec::AV1].Add(TRefCountPtr<EpicRtcVideoCodecInfoInterface>(new FVideoCodecInfo(
+			Codecs[EVideoCodec::AV1].Add(TRefCountPtr<EpicRtcVideoCodecInfoInterface>(new FEpicRtcVideoCodecInfo(
 				EpicRtcVideoCodec::AV1,
 				UE::PixelStreaming2::IsHardwareDecoderSupported<FVideoDecoderConfigAV1>())));
 		}
diff --git a/Engine/Plugins/Media/PixelStreaming2/Source/PixelStreaming2/Private/EpicRtcVideoEncoder.cpp b/Engine/Plugins/Media/PixelStreaming2/Source/PixelStreaming2/Private/EpicRtcVideoEncoder.cpp
--- a/Engine/Plugins/Media/PixelStreaming2/Source/PixelStreaming2/Private/EpicRtcVideoEncoder.cpp
+++ b/Engine/Plugins/Media/PixelStreaming2/Source/PixelStreaming2/Private/EpicRtcVideoEncoder.cpp
@@ -11,8 +11,8 @@
 #include "PixelCaptureOutputFrameI420.h"
 #include "PixelCaptureOutputFrameRHI.h"
 #include "PixelStreaming2PluginSettings.h"
-#include "ToStringExtensions.h"
-#include "VideoUtils.h"
+#include "UtilsString.h"
+#include "UtilsVideo.h"
 
 #include "epic_rtc/core/video/video_buffer.h"
 #include "epic_rtc/core/video/video_encoder_callback.h"
@@ -473,14 +473,14 @@ namespace UE::PixelStreaming2
 				CodecSpecificInfo._endOfPicture = Packet.CodecSpecificInfo.bEndOfPicture;
 				if (Packet.CodecSpecificInfo.GenericFrameInfo.IsSet())
 				{
-					CodecSpecificInfo._genericFrameInfo = new FGenericFrameInfoWrapper(Packet.CodecSpecificInfo.GenericFrameInfo.GetValue());
+					CodecSpecificInfo._genericFrameInfo = new FEpicRtcGenericFrameInfo(Packet.CodecSpecificInfo.GenericFrameInfo.GetValue());
 					CodecSpecificInfo._genericFrameInfo->AddRef();
 					CodecSpecificInfo._hasGenericFrameInfo = true;
 				}
 
 				if (Packet.CodecSpecificInfo.TemplateStructure.IsSet())
 				{
-					CodecSpecificInfo._templateStructure = new FFrameDependencyStructureWrapper(Packet.CodecSpecificInfo.TemplateStructure.GetValue());
+					CodecSpecificInfo._templateStructure = new FEpicRtcFrameDependencyStructure(Packet.CodecSpecificInfo.TemplateStructure.GetValue());
 					CodecSpecificInfo._templateStructure->AddRef();
 					CodecSpecificInfo._hasTemplateStructure = true;
 				}
diff --git a/Engine/Plugins/Media/PixelStreaming2/Source/PixelStreaming2/Private/EpicRtcVideoEncoder.h b/Engine/Plugins/Media/PixelStreaming2/Source/PixelStreaming2/Private/EpicRtcVideoEncoder.h
--- a/Engine/Plugins/Media/PixelStreaming2/Source/PixelStreaming2/Private/EpicRtcVideoEncoder.h
+++ b/Engine/Plugins/Media/PixelStreaming2/Source/PixelStreaming2/Private/EpicRtcVideoEncoder.h
@@ -3,7 +3,6 @@
 #pragma once
 
 #include "AVConstants.h"
-#include "EpicRtcMemory.h"
 #include "EpicRtcVideoCommon.h"
 #include "HAL/PlatformFileManager.h"
 #include "Video/Encoders/Configs/VideoEncoderConfigAV1.h"
@@ -25,7 +24,7 @@ namespace UE::PixelStreaming2
 		TEpicRtcVideoEncoder(EpicRtcVideoCodecInfoInterface* InCodecInfo);
 		virtual ~TEpicRtcVideoEncoder() override;
 
-		/* Begin EpicRtcVideoEncoderInterface */
+		// Begin EpicRtcVideoEncoderInterface
 		virtual EpicRtcStringView		  GetName() const override;
 		virtual EpicRtcVideoEncoderConfig GetConfig() const override;
 		virtual EpicRtcMediaResult		  SetConfig(const EpicRtcVideoEncoderConfig& VideoEncoderConfig) override;
@@ -33,7 +32,7 @@ namespace UE::PixelStreaming2
 		virtual EpicRtcMediaResult		  Encode(const EpicRtcVideoFrame& VideoFrame, EpicRtcVideoFrameTypeArrayInterface* FrameTypes) override;
 		virtual void					  RegisterCallback(EpicRtcVideoEncoderCallbackInterface* InCallback) override;
 		virtual void					  Reset() override;
-		/* End EpicRtcVideoEncoderInterface */
+		// End EpicRtcVideoEncoderInterface
 
 	private:
 		EpicRtcVideoEncoderConfig								EncoderConfig;
@@ -65,10 +64,10 @@ namespace UE::PixelStreaming2
 		FDelegateHandle DelegateHandle;
 
 	public:
-		/* Begin EpicRtcRefCountInterface */
+		// Begin EpicRtcRefCountInterface
 		virtual uint32_t AddRef() override final { return TRefCountingMixin<TEpicRtcVideoEncoder>::AddRef(); }
 		virtual uint32_t Release() override final { return TRefCountingMixin<TEpicRtcVideoEncoder>::Release(); }
 		virtual uint32_t Count() const override final { return TRefCountingMixin<TEpicRtcVideoEncoder>::GetRefCount(); }
-		/* End EpicRtcRefCountInterface */
+		// End EpicRtcRefCountInterface
 	};
 } // namespace UE::PixelStreaming2
diff --git a/Engine/Plugins/Media/PixelStreaming2/Source/PixelStreaming2/Private/EpicRtcVideoEncoderInitializer.cpp b/Engine/Plugins/Media/PixelStreaming2/Source/PixelStreaming2/Private/EpicRtcVideoEncoderInitializer.cpp
--- a/Engine/Plugins/Media/PixelStreaming2/Source/PixelStreaming2/Private/EpicRtcVideoEncoderInitializer.cpp
+++ b/Engine/Plugins/Media/PixelStreaming2/Source/PixelStreaming2/Private/EpicRtcVideoEncoderInitializer.cpp
@@ -8,9 +8,9 @@
 #include "NvmlEncoder.h"
 #include "PixelStreaming2Delegates.h"
 #include "PixelStreaming2PluginSettings.h"
-#include "ToStringExtensions.h"
-#include "Utils.h"
-#include "VideoUtils.h"
+#include "UtilsCommon.h"
+#include "UtilsString.h"
+#include "UtilsVideo.h"
 #include "Video/Encoders/Configs/VideoEncoderConfigAV1.h"
 #include "Video/Encoders/Configs/VideoEncoderConfigH264.h"
 #include "Video/Encoders/Configs/VideoEncoderConfigVP8.h"
@@ -192,7 +192,7 @@ namespace UE::PixelStreaming2
 		// UE doesn't support the automatic conversion of a RefCountPtr from a derived type to a base type.
 		if (UE::PixelStreaming2::IsEncoderSupported<FVideoEncoderConfigVP8>())
 		{
-			Codecs[EVideoCodec::VP8].Add(TRefCountPtr<EpicRtcVideoCodecInfoInterface>(new FVideoCodecInfo(
+			Codecs[EVideoCodec::VP8].Add(TRefCountPtr<EpicRtcVideoCodecInfoInterface>(new FEpicRtcVideoCodecInfo(
 				EpicRtcVideoCodec::VP8,
 				UE::PixelStreaming2::IsHardwareEncoderSupported<FVideoEncoderConfigVP8>(),
 				nullptr,
@@ -202,30 +202,32 @@ namespace UE::PixelStreaming2
 		if (UE::PixelStreaming2::IsEncoderSupported<FVideoEncoderConfigVP9>())
 		{
 			using namespace UE::AVCodecCore::VP9;
-			Codecs[EVideoCodec::VP9].Add(TRefCountPtr<EpicRtcVideoCodecInfoInterface>(new FVideoCodecInfo(
+			Codecs[EVideoCodec::VP9].Add(TRefCountPtr<EpicRtcVideoCodecInfoInterface>(new FEpicRtcVideoCodecInfo(
 				EpicRtcVideoCodec::VP9,
 				UE::PixelStreaming2::IsHardwareEncoderSupported<FVideoEncoderConfigVP9>(),
 				UE::PixelStreaming2::CreateVP9Format(EProfile::Profile0),
 				new FEpicRtcScalabilityModeArray(UE::PixelStreaming2::AllScalabilityModes))));
 
-			/* Only advertise profile 0 until EpicRtc provides us with a way to extract the negotiated profile in the encoder config
-			Codecs[EVideoCodec::VP9].Add(TRefCountPtr<EpicRtcVideoCodecInfoInterface>(new FVideoCodecInfo(
-				EpicRtcVideoCodec::VP9,
-				UE::PixelStreaming2::IsHardwareEncoderSupported<FVideoEncoderConfigVP9>(),
-				UE::PixelStreaming2::CreateVP9Format(EProfile::Profile2),
-				new FEpicRtcScalabilityModeArray(AllScalabilityModes))));
-			*/
+			/**
+			 * Only advertise profile 0 until EpicRtc provides us with a way to extract the negotiated profile in the encoder config
+			 * 
+			 * Codecs[EVideoCodec::VP9].Add(TRefCountPtr<EpicRtcVideoCodecInfoInterface>(new FEpicRtcVideoCodecInfo(
+			 *	EpicRtcVideoCodec::VP9,
+			 *	UE::PixelStreaming2::IsHardwareEncoderSupported<FVideoEncoderConfigVP9>(),
+			 *	UE::PixelStreaming2::CreateVP9Format(EProfile::Profile2),
+			 *	new FEpicRtcScalabilityModeArray(AllScalabilityModes))));
+			 */
 		}
 
 		if (UE::PixelStreaming2::IsEncoderSupported<FVideoEncoderConfigH264>())
 		{
 			using namespace UE::AVCodecCore::H264;
-			Codecs[EVideoCodec::H264].Add(TRefCountPtr<EpicRtcVideoCodecInfoInterface>(new FVideoCodecInfo(
+			Codecs[EVideoCodec::H264].Add(TRefCountPtr<EpicRtcVideoCodecInfoInterface>(new FEpicRtcVideoCodecInfo(
 				EpicRtcVideoCodec::H264,
 				UE::PixelStreaming2::IsHardwareEncoderSupported<FVideoEncoderConfigH264>(),
 				UE::PixelStreaming2::CreateH264Format(EH264Profile::ConstrainedBaseline, EH264Level::Level_3_1),
 				new FEpicRtcScalabilityModeArray({ EScalabilityMode::L1T1 }))));
-			Codecs[EVideoCodec::H264].Add(TRefCountPtr<EpicRtcVideoCodecInfoInterface>(new FVideoCodecInfo(
+			Codecs[EVideoCodec::H264].Add(TRefCountPtr<EpicRtcVideoCodecInfoInterface>(new FEpicRtcVideoCodecInfo(
 				EpicRtcVideoCodec::H264,
 				UE::PixelStreaming2::IsHardwareEncoderSupported<FVideoEncoderConfigH264>(),
 				UE::PixelStreaming2::CreateH264Format(EH264Profile::Baseline, EH264Level::Level_3_1),
@@ -234,7 +236,7 @@ namespace UE::PixelStreaming2
 
 		if (UE::PixelStreaming2::IsEncoderSupported<FVideoEncoderConfigAV1>())
 		{
-			Codecs[EVideoCodec::AV1].Add(TRefCountPtr<EpicRtcVideoCodecInfoInterface>(new FVideoCodecInfo(
+			Codecs[EVideoCodec::AV1].Add(TRefCountPtr<EpicRtcVideoCodecInfoInterface>(new FEpicRtcVideoCodecInfo(
 				EpicRtcVideoCodec::AV1,
 				UE::PixelStreaming2::IsHardwareEncoderSupported<FVideoEncoderConfigAV1>(),
 				nullptr,
diff --git a/Engine/Plugins/Media/PixelStreaming2/Source/PixelStreaming2/Private/EpicRtcWebsocket.cpp b/Engine/Plugins/Media/PixelStreaming2/Source/PixelStreaming2/Private/EpicRtcWebsocket.cpp
--- a/Engine/Plugins/Media/PixelStreaming2/Source/PixelStreaming2/Private/EpicRtcWebsocket.cpp
+++ b/Engine/Plugins/Media/PixelStreaming2/Source/PixelStreaming2/Private/EpicRtcWebsocket.cpp
@@ -6,8 +6,8 @@
 #include "IWebSocket.h"
 #include "PixelStreaming2PluginSettings.h"
 #include "Serialization/JsonSerializer.h"
-#include "ToStringExtensions.h"
-#include "VideoUtils.h"
+#include "UtilsString.h"
+#include "UtilsVideo.h"
 #include "WebSocketsModule.h"
 
 namespace UE::PixelStreaming2
diff --git a/Engine/Plugins/Media/PixelStreaming2/Source/PixelStreaming2/Private/FreezeFrame.cpp b/Engine/Plugins/Media/PixelStreaming2/Source/PixelStreaming2/Private/FreezeFrame.cpp
--- a/Engine/Plugins/Media/PixelStreaming2/Source/PixelStreaming2/Private/FreezeFrame.cpp
+++ b/Engine/Plugins/Media/PixelStreaming2/Source/PixelStreaming2/Private/FreezeFrame.cpp
@@ -21,7 +21,7 @@
 #include "ScreenPass.h"
 #include "ScreenRendering.h"
 #include "TextureResource.h"
-#include "Utils.h"
+#include "UtilsCommon.h"
 
 namespace UE::PixelStreaming2
 {
@@ -43,7 +43,7 @@ namespace UE::PixelStreaming2
 		RemoveFreezeFrameBinding();
 	}
 
-	/*
+	/**
 	 * Add the commands to the RHI command list to copy a texture from source to dest - even if the format is different.
 	 * Assumes SourceTexture is in ERHIAccess::CopySrc and DestTexture is in ERHIAccess::CopyDest
 	 */
diff --git a/Engine/Plugins/Media/PixelStreaming2/Source/PixelStreaming2/Private/FreezeFrame.h b/Engine/Plugins/Media/PixelStreaming2/Source/PixelStreaming2/Private/FreezeFrame.h
--- a/Engine/Plugins/Media/PixelStreaming2/Source/PixelStreaming2/Private/FreezeFrame.h
+++ b/Engine/Plugins/Media/PixelStreaming2/Source/PixelStreaming2/Private/FreezeFrame.h
@@ -12,7 +12,7 @@
 
 namespace UE::PixelStreaming2
 {
-	/*
+	/**
 	 * The user of the Pixel Streaming application can trigger a freeze frame to happen at any time during stream.
 	 * What a freeze frame will do is either:
 	 * 1) Mode 1 - Send a known texture/image (effectively a still image, e.g. a loading screen).
diff --git a/Engine/Plugins/Media/PixelStreaming2/Source/PixelStreaming2/Private/PixelStreaming2Module.cpp b/Engine/Plugins/Media/PixelStreaming2/Source/PixelStreaming2/Private/PixelStreaming2Module.cpp
--- a/Engine/Plugins/Media/PixelStreaming2/Source/PixelStreaming2/Private/PixelStreaming2Module.cpp
+++ b/Engine/Plugins/Media/PixelStreaming2/Source/PixelStreaming2/Private/PixelStreaming2Module.cpp
@@ -18,7 +18,7 @@
 #include "Slate/SceneViewport.h"
 #include "Streamer.h"
 #include "UObject/UObjectIterator.h"
-#include "Utils.h"
+#include "UtilsCommon.h"
 
 #if PLATFORM_LINUX
 	#include "CudaModule.h"
@@ -59,7 +59,7 @@ THIRD_PARTY_INCLUDES_END
 #include "VideoProducerMediaCapture.h"
 #include "Engine/GameEngine.h"
 #include "Stats.h"
-#include "ToStringExtensions.h"
+#include "UtilsString.h"
 
 #include "EpicRtcAllocator.h"
 #include "EpicRtcAudioMixingCapturer.h"
@@ -96,7 +96,7 @@ namespace UE::PixelStreaming2
 
 	FUtf8String FPixelStreaming2Module::EpicRtcConferenceName("pixel_streaming_conference_instance");
 
-	/*
+	/**
 	 * Stats logger - as turned on/off by CVarPixelStreaming2LogStats
 	 */
 	void ConsumeStat(FString PlayerId, FName StatName, float StatValue)
diff --git a/Engine/Plugins/Media/PixelStreaming2/Source/PixelStreaming2/Private/PixelStreaming2Module.h b/Engine/Plugins/Media/PixelStreaming2/Source/PixelStreaming2/Private/PixelStreaming2Module.h
--- a/Engine/Plugins/Media/PixelStreaming2/Source/PixelStreaming2/Private/PixelStreaming2Module.h
+++ b/Engine/Plugins/Media/PixelStreaming2/Source/PixelStreaming2/Private/PixelStreaming2Module.h
@@ -19,7 +19,7 @@ namespace UE::PixelStreaming2
 	class FVideoSourceGroup;
 	class FEpicRtcWebsocketFactory;
 
-	/*
+	/**
 	 * This plugin allows the back buffer to be sent as a compressed video across a network.
 	 */
 	class FPixelStreaming2Module : public IPixelStreaming2Module
@@ -29,7 +29,7 @@ namespace UE::PixelStreaming2
 
 		virtual ~FPixelStreaming2Module() = default;
 
-		/** IPixelStreaming2Module implementation */
+		// Begin IPixelStreaming2Module
 		virtual FReadyEvent&							  OnReady() override;
 		virtual bool									  IsReady() override;
 		virtual bool									  StartStreaming() override;
@@ -49,17 +49,17 @@ namespace UE::PixelStreaming2
 		virtual void								 RemoveInputComponent(UPixelStreaming2Input* InInputComponent);
 		virtual const TArray<UPixelStreaming2Input*> GetInputComponents();
 		virtual void								 ForEachStreamer(const TFunction<void(TSharedPtr<IPixelStreaming2Streamer>)>& Func) override;
-		/** End IPixelStreaming2Module implementation */
+		// End IPixelStreaming2Module
 
 		TSharedPtr<class FEpicRtcAudioMixingCapturer> GetAudioCapturer();
 		TRefCountPtr<EpicRtcConferenceInterface>	  GetEpicRtcConference() { return EpicRtcConference; }
 		TRefCountPtr<FEpicRtcStatsCollector>		  GetStatsCollector() { return StatsCollector; }
 
 	private:
-		/** IModuleInterface implementation */
+		// Begin IModuleInterface
 		void StartupModule() override;
 		void ShutdownModule() override;
-		/** End IModuleInterface implementation */
+		// End IModuleInterface
 
 		// Own methods
 		void InitDefaultStreamer();
diff --git a/Engine/Plugins/Media/PixelStreaming2/Source/PixelStreaming2/Private/PixelStreaming2Utils.cpp b/Engine/Plugins/Media/PixelStreaming2/Source/PixelStreaming2/Private/PixelStreaming2Utils.cpp
--- a/Engine/Plugins/Media/PixelStreaming2/Source/PixelStreaming2/Private/PixelStreaming2Utils.cpp
+++ b/Engine/Plugins/Media/PixelStreaming2/Source/PixelStreaming2/Private/PixelStreaming2Utils.cpp
@@ -32,7 +32,7 @@ namespace UE::PixelStreaming2
 
 	void ExtractJsonFromDescriptor(FString Descriptor, FString FieldName, FString& OutStringValue, bool& OutSuccess)
 	{
-		/*
+		/**
 		 * ExtractJsonFromDescriptor supports parsing nested objects, as well as keys that contain a period.
 		 * eg
 		 * { "Encoder.MinQP": val }
@@ -73,5 +73,4 @@ namespace UE::PixelStreaming2
 		}
 	}
 
-} // UE::PixelStreaming2
-
+} // namespace UE::PixelStreaming2
diff --git a/Engine/Plugins/Media/PixelStreaming2/Source/PixelStreaming2/Private/PlayerContext.h b/Engine/Plugins/Media/PixelStreaming2/Source/PixelStreaming2/Private/PlayerContext.h
--- a/Engine/Plugins/Media/PixelStreaming2/Source/PixelStreaming2/Private/PlayerContext.h
+++ b/Engine/Plugins/Media/PixelStreaming2/Source/PixelStreaming2/Private/PlayerContext.h
@@ -9,7 +9,7 @@
 #include "EpicRtcVideoSource.h"
 #include "RTCStatsCollector.h"
 #include "Templates/RefCounting.h"
-#include "Utils.h"
+#include "UtilsCommon.h"
 
 #include "epic_rtc/core/audio/audio_track.h"
 #include "epic_rtc/core/video/video_track.h"
diff --git a/Engine/Plugins/Media/PixelStreaming2/Source/PixelStreaming2/Private/RTCStatsCollector.cpp b/Engine/Plugins/Media/PixelStreaming2/Source/PixelStreaming2/Private/RTCStatsCollector.cpp
--- a/Engine/Plugins/Media/PixelStreaming2/Source/PixelStreaming2/Private/RTCStatsCollector.cpp
+++ b/Engine/Plugins/Media/PixelStreaming2/Source/PixelStreaming2/Private/RTCStatsCollector.cpp
@@ -6,7 +6,7 @@
 #include "PixelStreaming2PluginSettings.h"
 #include "PixelStreaming2StatNames.h"
 #include "Streamer.h"
-#include "ToStringExtensions.h"
+#include "UtilsString.h"
 
 namespace UE::PixelStreaming2
 {
@@ -122,7 +122,7 @@ namespace UE::PixelStreaming2
 
 		LastCalculationCycles = FPlatformTime::Cycles64();
 	}
-	/*
+	/**
 	 * ---------- FRTPLocalVideoTrackSink ----------
 	 */
 	FRTCStatsCollector::FRTPLocalVideoTrackStatsSink::FRTPLocalVideoTrackStatsSink()
@@ -431,7 +431,7 @@ namespace UE::PixelStreaming2
 			PostProcess(PSStats, StatsId, SecondsDelta);
 		}
 	}
-	/*
+	/**
 	 * ---------- FRTPLocalAudioTrackStatsSink ----------
 	 */
 	FRTCStatsCollector::FRTPLocalAudioTrackStatsSink::FRTPLocalAudioTrackStatsSink()
@@ -663,7 +663,7 @@ namespace UE::PixelStreaming2
 		PostProcess(PSStats, PeerId, SecondsDelta);
 	}
 
-	/*
+	/**
 	 * ---------- FRTPRemoteTrackStatsSink ----------
 	 */
 	FRTCStatsCollector::FRTPRemoteTrackStatsSink::FRTPRemoteTrackStatsSink(FName InCategory)
@@ -932,7 +932,7 @@ namespace UE::PixelStreaming2
 		PostProcess(PSStats, PeerId, SecondsDelta);
 	}
 
-	/*
+	/**
 	 * ---------- FRTPVideoSourceSink ----------
 	 */
 	FRTCStatsCollector::FVideoSourceStatsSink::FVideoSourceStatsSink()
@@ -967,7 +967,7 @@ namespace UE::PixelStreaming2
 		PostProcess(PSStats, PeerId, SecondsDelta);
 	}
 
-	/*
+	/**
 	 * ---------- FRTPAudioSourceSink ----------
 	 */
 	FRTCStatsCollector::FAudioSourceStatsSink::FAudioSourceStatsSink()
@@ -1006,7 +1006,7 @@ namespace UE::PixelStreaming2
 		PostProcess(PSStats, PeerId, SecondsDelta);
 	}
 
-	/*
+	/**
 	 * ----------- FDataChannelSink -----------
 	 */
 	FRTCStatsCollector::FDataTrackStatsSink::FDataTrackStatsSink()
@@ -1056,7 +1056,8 @@ namespace UE::PixelStreaming2
 
 		PostProcess(PSStats, PeerId, SecondsDelta);
 	}
-	/*
+
+	/**
 	 * ---------- FRTPAudioSourceSink ----------
 	 */
 	FRTCStatsCollector::FCandidatePairStatsSink::FCandidatePairStatsSink()
@@ -1095,7 +1096,7 @@ namespace UE::PixelStreaming2
 		PostProcess(PSStats, PeerId, SecondsDelta);
 	}
 
-	/*
+	/**
 	 * ---------- FRTCTrackedStat -------------------
 	 */
 	FRTCStatsCollector::FRTCTrackedStat::FRTCTrackedStat(FName StatName, FName Alias, int NDecimalPlaces, uint8 DisplayFlags)
@@ -1121,7 +1122,7 @@ namespace UE::PixelStreaming2
 		LatestStat.StatValue = InValue;
 	}
 
-	/*
+	/**
 	 * --------- FStatsSink ------------------------
 	 */
 	FRTCStatsCollector::FStatsSink::FStatsSink(FName InCategory)
diff --git a/Engine/Plugins/Media/PixelStreaming2/Source/PixelStreaming2/Private/RTCStatsCollector.h b/Engine/Plugins/Media/PixelStreaming2/Source/PixelStreaming2/Private/RTCStatsCollector.h
--- a/Engine/Plugins/Media/PixelStreaming2/Source/PixelStreaming2/Private/RTCStatsCollector.h
+++ b/Engine/Plugins/Media/PixelStreaming2/Source/PixelStreaming2/Private/RTCStatsCollector.h
@@ -35,7 +35,7 @@ namespace UE::PixelStreaming2
 		void Process(const EpicRtcConnectionStats& InStats);
 
 	private:
-		/*
+		/**
 		 * ---------- FRTCTrackedStat -------------------
 		 * Tracks a stat from WebRTC. Stores the current value and previous value.
 		 */
@@ -146,7 +146,7 @@ namespace UE::PixelStreaming2
 			void Process(const EpicRtcAudioSourceStats& InStats, const FString& PeerId, double SecondsDelta);
 		};
 
-		/*
+		/**
 		 * ---------- FDataChannelSink -------------------
 		 */
 		class FDataTrackStatsSink : public FStatsSink
diff --git a/Engine/Plugins/Media/PixelStreaming2/Source/PixelStreaming2/Private/Stats.h b/Engine/Plugins/Media/PixelStreaming2/Source/PixelStreaming2/Private/Stats.h
--- a/Engine/Plugins/Media/PixelStreaming2/Source/PixelStreaming2/Private/Stats.h
+++ b/Engine/Plugins/Media/PixelStreaming2/Source/PixelStreaming2/Private/Stats.h
@@ -52,8 +52,8 @@ namespace UE::PixelStreaming2
 		bool			 bSmooth;
 		double			 LastEMA = 0;
 		int				 NumSamples = 0;
-		uint8			 DisplayFlags = EDisplayFlags::TEXT; /* Some stats we only wish to store or broadcast, but not display */
-		TOptional<FName> Alias;								 /* Some stats need an alias that they are stored by/queried by to disambiguate them from other stats */
+		uint8			 DisplayFlags = EDisplayFlags::TEXT; // Some stats we only wish to store or broadcast, but not display
+		TOptional<FName> Alias;								 // Some stats need an alias that they are stored by/queried by to disambiguate them from other stats
 	};
 
 	FORCEINLINE uint32 GetTypeHash(const FStatData& Obj)
diff --git a/Engine/Plugins/Media/PixelStreaming2/Source/PixelStreaming2/Private/Streamer.cpp b/Engine/Plugins/Media/PixelStreaming2/Source/PixelStreaming2/Private/Streamer.cpp
--- a/Engine/Plugins/Media/PixelStreaming2/Source/PixelStreaming2/Private/Streamer.cpp
+++ b/Engine/Plugins/Media/PixelStreaming2/Source/PixelStreaming2/Private/Streamer.cpp
@@ -25,9 +25,9 @@
 #include "EpicRtcDataTrack.h"
 #include "IPixelStreaming2Module.h"
 #include "IPixelStreaming2InputModule.h"
-#include "ToStringExtensions.h"
 #include "CoreGlobals.h"
-#include "VideoUtils.h"
+#include "UtilsString.h"
+#include "UtilsVideo.h"
 
 namespace UE::PixelStreaming2
 {
@@ -628,7 +628,7 @@ namespace UE::PixelStreaming2
 						return false;
 					}
 					InputControllingId = PlayerId;
-					PlayerContext.DataTrack->SendMessage(EPixelStreaming2FromStreamerMessage::InputControlOwnership, 1 /*ControlsInput*/);
+					PlayerContext.DataTrack->SendMessage(EPixelStreaming2FromStreamerMessage::InputControlOwnership, 1 /* ControlsInput */);
 					return true;
 				});
 			}
@@ -718,7 +718,7 @@ namespace UE::PixelStreaming2
 			const uint8 ControlsInput = (GetEnumFromCVar<EInputControllerMode>(UPixelStreaming2PluginSettings::CVarInputController) == EInputControllerMode::Host) ? (PlayerId == InputControllingId) : 1;
 			// Even though the QualityController feature is removed we send it for backwards compatibility with older frontends (can probably remove 2 versions after 5.5)
 			PlayerContext->DataTrack->SendMessage(EPixelStreaming2FromStreamerMessage::InputControlOwnership, ControlsInput);
-			PlayerContext->DataTrack->SendMessage(EPixelStreaming2FromStreamerMessage::QualityControlOwnership, 1 /*True*/);
+			PlayerContext->DataTrack->SendMessage(EPixelStreaming2FromStreamerMessage::QualityControlOwnership, 1 /* True */);
 		}
 	}
 
diff --git a/Engine/Plugins/Media/PixelStreaming2/Source/PixelStreaming2/Private/Streamer.h b/Engine/Plugins/Media/PixelStreaming2/Source/PixelStreaming2/Private/Streamer.h
--- a/Engine/Plugins/Media/PixelStreaming2/Source/PixelStreaming2/Private/Streamer.h
+++ b/Engine/Plugins/Media/PixelStreaming2/Source/PixelStreaming2/Private/Streamer.h
@@ -128,7 +128,7 @@ namespace UE::PixelStreaming2
 		TSharedPtr<FEpicRtcManager> EpicRtcManager;
 
 	public:
-		/* Begin FEpicRtcManager Callbacks */
+		// Begin FEpicRtcManager Callbacks
 		void OnSessionStateUpdate(const EpicRtcSessionState StateUpdate);
 		void OnSessionErrorUpdate(const EpicRtcErrorCode ErrorUpdate);
 		void OnSessionRoomsAvailableUpdate(EpicRtcStringArrayInterface* RoomsList);
@@ -156,6 +156,6 @@ namespace UE::PixelStreaming2
 		void OnDataTrackRemoved(EpicRtcDataTrackInterface* DataTrack);
 		void OnDataTrackState(EpicRtcDataTrackInterface* DataTrack, const EpicRtcTrackState State);
 		void OnDataTrackMessage(EpicRtcDataTrackInterface* DataTrack);
-		/* End FEpicRtcManager Callbacks */
+		// End FEpicRtcManager Callbacks
 	};
 } // namespace UE::PixelStreaming2
diff --git a/Engine/Plugins/Media/PixelStreaming2/Source/PixelStreaming2/Private/StreamerReconnectTimer.h b/Engine/Plugins/Media/PixelStreaming2/Source/PixelStreaming2/Private/StreamerReconnectTimer.h
--- a/Engine/Plugins/Media/PixelStreaming2/Source/PixelStreaming2/Private/StreamerReconnectTimer.h
+++ b/Engine/Plugins/Media/PixelStreaming2/Source/PixelStreaming2/Private/StreamerReconnectTimer.h
@@ -24,12 +24,12 @@ namespace UE::PixelStreaming2
 		void Start(TWeakPtr<IPixelStreaming2Streamer> InWeakStreamer);
 		void Stop();
 
-		/* Begin FTickableGameObject */
+		// Begin FTickableGameObject
 		virtual void Tick (float DeltaTime) override;
 		virtual bool IsTickableInEditor() const override { return true; }
 		virtual bool IsTickableWhenPaused() const override { return true; }
 		FORCEINLINE TStatId GetStatId() const { RETURN_QUICK_DECLARE_CYCLE_STAT(PixelStreaming2StreamerReconnectTimer, STATGROUP_Tickables); }
-		/* End FTickableGameObject */
+		// End FTickableGameObject
 
 	private:
 		TWeakPtr<IPixelStreaming2Streamer> WeakStreamer;
diff --git a/Engine/Plugins/Media/PixelStreaming2/Source/PixelStreaming2/Private/Tests/EpicRtcConferenceTest.cpp b/Engine/Plugins/Media/PixelStreaming2/Source/PixelStreaming2/Private/Tests/EpicRtcConferenceTest.cpp
--- a/Engine/Plugins/Media/PixelStreaming2/Source/PixelStreaming2/Private/Tests/EpicRtcConferenceTest.cpp
+++ b/Engine/Plugins/Media/PixelStreaming2/Source/PixelStreaming2/Private/Tests/EpicRtcConferenceTest.cpp
@@ -4,7 +4,7 @@
 #include "Logging.h"
 #include "Misc/AutomationTest.h"
 #include "Tests/AutomationCommon.h"
-#include "ToStringExtensions.h"
+#include "UtilsString.h"
 
 #include "epic_rtc/core/platform.h"
 
diff --git a/Engine/Plugins/Media/PixelStreaming2/Source/PixelStreaming2/Private/Tests/EpicRtcRoomTest.cpp b/Engine/Plugins/Media/PixelStreaming2/Source/PixelStreaming2/Private/Tests/EpicRtcRoomTest.cpp
--- a/Engine/Plugins/Media/PixelStreaming2/Source/PixelStreaming2/Private/Tests/EpicRtcRoomTest.cpp
+++ b/Engine/Plugins/Media/PixelStreaming2/Source/PixelStreaming2/Private/Tests/EpicRtcRoomTest.cpp
@@ -5,7 +5,7 @@
 #include "Logging.h"
 #include "Misc/AutomationTest.h"
 #include "Tests/AutomationCommon.h"
-#include "ToStringExtensions.h"
+#include "UtilsString.h"
 
 #if WITH_DEV_AUTOMATION_TESTS
 
diff --git a/Engine/Plugins/Media/PixelStreaming2/Source/PixelStreaming2/Private/Tests/EpicRtcSessionTest.cpp b/Engine/Plugins/Media/PixelStreaming2/Source/PixelStreaming2/Private/Tests/EpicRtcSessionTest.cpp
--- a/Engine/Plugins/Media/PixelStreaming2/Source/PixelStreaming2/Private/Tests/EpicRtcSessionTest.cpp
+++ b/Engine/Plugins/Media/PixelStreaming2/Source/PixelStreaming2/Private/Tests/EpicRtcSessionTest.cpp
@@ -5,7 +5,7 @@
 #include "Logging.h"
 #include "Misc/AutomationTest.h"
 #include "Tests/AutomationCommon.h"
-#include "ToStringExtensions.h"
+#include "UtilsString.h"
 
 #include "epic_rtc/core/platform.h"
 
diff --git a/Engine/Plugins/Media/PixelStreaming2/Source/PixelStreaming2/Private/Tests/EpicRtcUtils.h b/Engine/Plugins/Media/PixelStreaming2/Source/PixelStreaming2/Private/Tests/EpicRtcUtils.h
--- a/Engine/Plugins/Media/PixelStreaming2/Source/PixelStreaming2/Private/Tests/EpicRtcUtils.h
+++ b/Engine/Plugins/Media/PixelStreaming2/Source/PixelStreaming2/Private/Tests/EpicRtcUtils.h
@@ -7,7 +7,7 @@
 #include "IWebSocket.h"
 #include "Logging.h"
 #include "Misc/AutomationTest.h"
-#include "ToStringExtensions.h"
+#include "UtilsString.h"
 
 #include "epic_rtc/common/common.h"
 #include "epic_rtc/core/conference.h"
@@ -15,33 +15,13 @@
 
 namespace UE::PixelStreaming2
 {
-	inline FString ToString(EpicRtcRoomState State)
-	{
-		switch (State)
-		{
-			case EpicRtcRoomState::New:
-				return TEXT("New");
-			case EpicRtcRoomState::Pending:
-				return TEXT("Pending");
-			case EpicRtcRoomState::Joined:
-				return TEXT("Joined");
-			case EpicRtcRoomState::Left:
-				return TEXT("Left");
-			case EpicRtcRoomState::Failed:
-				return TEXT("Failed");
-			case EpicRtcRoomState::Exiting:
-				return TEXT("Exiting");
-		}
-		return TEXT("Unknown");
-	}
-
 	// A mock manager class for tests to receive callbacks from EpicRtc. Typically, the controlling class will inherit from FEpicRtcManager
 	// and implement the methods itself (see streamer.cpp). However, we can't force the tests to inherit the class, so instead we have the
 	// mock manager and the test bodies bind to the events they're interested in
 	class FMockManager : public FEpicRtcManager
 	{
 	public:
-		/* Begin FEpicRtcAudioTrackObserver Callbacks */
+		// Begin FEpicRtcAudioTrackObserver Callbacks
 		void OnAudioTrackMuted(EpicRtcAudioTrackInterface* AudioTrack, EpicRtcBool bIsMuted)
 		{
 		}
@@ -54,9 +34,9 @@ namespace UE::PixelStreaming2
 		void OnAudioTrackState(EpicRtcAudioTrackInterface* AudioTrack, const EpicRtcTrackState State)
 		{
 		}
-		/* End FEpicRtcAudioTrackObserver Callbacks */
+		// End FEpicRtcAudioTrackObserver Callbacks
 
-		/* Begin FEpicRtcVideoTrackObserver Callbacks */
+		// Begin FEpicRtcVideoTrackObserver Callbacks
 		void OnVideoTrackMuted(EpicRtcVideoTrackInterface* VideoTrack, EpicRtcBool bIsMuted)
 		{
 		}
@@ -69,9 +49,9 @@ namespace UE::PixelStreaming2
 		void OnVideoTrackState(EpicRtcVideoTrackInterface* VideoTrack, const EpicRtcTrackState State)
 		{
 		}
-		/* End FEpicRtcVideoTrackObserver Callbacks */
+		// End FEpicRtcVideoTrackObserver Callbacks
 
-		/* Begin FEpicRtcDataTrackObserver Callbacks */
+		// Begin FEpicRtcDataTrackObserver Callbacks
 		void OnDataTrackRemoved(EpicRtcDataTrackInterface* DataTrack)
 		{
 		}
@@ -81,7 +61,7 @@ namespace UE::PixelStreaming2
 		void OnDataTrackMessage(EpicRtcDataTrackInterface* DataTrack)
 		{
 		}
-		/* End FEpicRtcDataTrackObserver Callbacks */
+		// End FEpicRtcDataTrackObserver Callbacks
 
 	public:
 		TRefCountPtr<EpicRtcConferenceInterface>& GetEpicRtcConference()
@@ -200,11 +180,11 @@ namespace UE::PixelStreaming2
 		}
 
 	public:
-		/* Begin EpicRtcRefCountInterface */
+		// Begin EpicRtcRefCountInterface
 		virtual uint32_t AddRef() override final { return TRefCountingMixin<FMockWebSocketFactory>::AddRef(); }
 		virtual uint32_t Release() override final { return TRefCountingMixin<FMockWebSocketFactory>::Release(); }
 		virtual uint32_t Count() const override final { return TRefCountingMixin<FMockWebSocketFactory>::GetRefCount(); }
-		/* End EpicRtcRefCountInterface */
+		// End EpicRtcRefCountInterface
 
 	private:
 		TRefCountPtr<EpicRtcWebsocketInterface> Websocket;
diff --git a/Engine/Plugins/Media/PixelStreaming2/Source/PixelStreaming2/Private/Tests/MockPlayer.cpp b/Engine/Plugins/Media/PixelStreaming2/Source/PixelStreaming2/Private/Tests/MockPlayer.cpp
--- a/Engine/Plugins/Media/PixelStreaming2/Source/PixelStreaming2/Private/Tests/MockPlayer.cpp
+++ b/Engine/Plugins/Media/PixelStreaming2/Source/PixelStreaming2/Private/Tests/MockPlayer.cpp
@@ -2,12 +2,12 @@
 
 #include "MockPlayer.h"
 
+#include "DefaultDataProtocol.h"
 #include "EpicRtcVideoEncoderInitializer.h"
 #include "EpicRtcVideoDecoderInitializer.h"
 #include "EpicRtcWebsocketFactory.h"
 #include "Logging.h"
-#include "ToStringExtensions.h"
-#include "DefaultDataProtocol.h"
+#include "UtilsString.h"
 
 #include "epic_rtc/core/platform.h"
 
diff --git a/Engine/Plugins/Media/PixelStreaming2/Source/PixelStreaming2/Private/ToStringExtensions.h b/Engine/Plugins/Media/PixelStreaming2/Source/PixelStreaming2/Private/ToStringExtensions.h
deleted file mode 100644
diff --git a/Engine/Plugins/Media/PixelStreaming2/Source/PixelStreaming2/Private/Utils.h b/Engine/Plugins/Media/PixelStreaming2/Source/PixelStreaming2/Private/Utils.h
deleted file mode 100644
diff --git a/Engine/Plugins/Media/PixelStreaming2/Source/PixelStreaming2/Private/UtilsCommon.h b/Engine/Plugins/Media/PixelStreaming2/Source/PixelStreaming2/Private/UtilsCommon.h
new file mode 100644
--- /dev/null
+++ b/Engine/Plugins/Media/PixelStreaming2/Source/PixelStreaming2/Private/UtilsCommon.h
@@ -0,0 +1,53 @@
+// Copyright Epic Games, Inc. All Rights Reserved.
+
+#pragma once
+
+#include "Async/Async.h"
+#include "CoreMinimal.h"
+#include "EpicRtcVideoCommon.h"
+#include "HAL/IConsoleManager.h"
+#include "Misc/CommandLine.h"
+#include "Video/CodecUtils/CodecUtilsH264.h"
+#include "Video/CodecUtils/CodecUtilsVP9.h"
+
+namespace UE::PixelStreaming2
+{
+	static const FString SFU_PLAYER_ID_PREFIX = FString(TEXT("SFU_"));
+
+	inline bool IsSFU(const FString& InPlayerId)
+	{
+		return InPlayerId.StartsWith(SFU_PLAYER_ID_PREFIX);
+	}
+
+	template <typename T>
+	void DoOnGameThread(T&& Func)
+	{
+		if (IsInGameThread())
+		{
+			Func();
+		}
+		else
+		{
+			AsyncTask(ENamedThreads::GameThread, [Func]() { Func(); });
+		}
+	}
+
+	template <typename T>
+	void DoOnGameThreadAndWait(uint32 Timeout, T&& Func)
+	{
+		if (IsInGameThread())
+		{
+			Func();
+		}
+		else
+		{
+			FEvent* TaskEvent = FPlatformProcess::GetSynchEventFromPool();
+			AsyncTask(ENamedThreads::GameThread, [Func, TaskEvent]() {
+				Func();
+				TaskEvent->Trigger();
+			});
+			TaskEvent->Wait(Timeout);
+			FPlatformProcess::ReturnSynchEventToPool(TaskEvent);
+		}
+	}
+} // namespace UE::PixelStreaming2
diff --git a/Engine/Plugins/Media/PixelStreaming2/Source/PixelStreaming2/Private/UtilsString.h b/Engine/Plugins/Media/PixelStreaming2/Source/PixelStreaming2/Private/UtilsString.h
new file mode 100644
--- /dev/null
+++ b/Engine/Plugins/Media/PixelStreaming2/Source/PixelStreaming2/Private/UtilsString.h
@@ -0,0 +1,189 @@
+// Copyright Epic Games, Inc. All Rights Reserved.
+
+#pragma once
+
+#include <string>
+#include "Containers/StringConv.h"
+#include "Containers/UnrealString.h"
+#include "Containers/Utf8String.h"
+#include "Dom/JsonObject.h"
+#include "Policies/PrettyJsonPrintPolicy.h"
+#include "Policies/CondensedJsonPrintPolicy.h"
+#include "Serialization/JsonSerializer.h"
+#include "epic_rtc/common/common.h"
+#include "epic_rtc/containers/epic_rtc_string_view.h"
+#include "epic_rtc/core/video/video_common.h"
+
+namespace UE::PixelStreaming2
+{
+	inline FString ToString(EpicRtcErrorCode Error)
+	{
+		switch (Error)
+		{
+			case EpicRtcErrorCode::Ok:
+			{
+				return TEXT("Ok");
+			}
+			case EpicRtcErrorCode::GeneralError:
+			{
+				return TEXT("GeneralError");
+			}
+			case EpicRtcErrorCode::BadState:
+			{
+				return TEXT("BadState");
+			}
+			case EpicRtcErrorCode::Timeout:
+			{
+				return TEXT("Timeout");
+			}
+			case EpicRtcErrorCode::Unsupported:
+			{
+				return TEXT("Unsupported");
+			}
+			case EpicRtcErrorCode::PlatformError:
+			{
+				return TEXT("PlatformError");
+			}
+			case EpicRtcErrorCode::FoundExistingPlatform:
+			{
+				return TEXT("FoundExistingPlatform");
+			}
+			case EpicRtcErrorCode::ConferenceAlreadyExists:
+			{
+				return TEXT("ConferenceAlreadyExists");
+			}
+			case EpicRtcErrorCode::ConferenceDoesNotExists:
+			{
+				return TEXT("ConferenceDoesNotExists");
+			}
+			case EpicRtcErrorCode::ImATeapot:
+			{
+				return TEXT("ImATeapot");
+			}
+			case EpicRtcErrorCode::ConferenceError:
+			{
+				return TEXT("ConferenceError");
+			}
+			case EpicRtcErrorCode::SessionAlreadyExists:
+			{
+				return TEXT("SessionAlreadyExists");
+			}
+			case EpicRtcErrorCode::SessionDoesNotExist:
+			{
+				return TEXT("SessionDoesNotExist");
+			}
+			case EpicRtcErrorCode::SessionError:
+			{
+				return TEXT("SessionError");
+			}
+			case EpicRtcErrorCode::SessionCannotConnect:
+			{
+				return TEXT("SessionCannotConnect");
+			}
+			case EpicRtcErrorCode::SessionDisconnected:
+			{
+				return TEXT("SessionDisconnected");
+			}
+			case EpicRtcErrorCode::SessionCannotCreateRoom:
+			{
+				return TEXT("SessionCannotCreateRoom");
+			}
+			default:
+			{
+				return TEXT("Unknown");
+			}
+		}
+	}
+
+	inline FString ToString(EpicRtcVideoCodec Codec)
+	{
+		switch (Codec)
+		{
+			case EpicRtcVideoCodec::AV1:
+			{
+				return TEXT("AV1");
+			}
+			case EpicRtcVideoCodec::H264:
+			{
+				return TEXT("H264");
+			}
+			case EpicRtcVideoCodec::VP8:
+			{
+				return TEXT("VP8");
+			}
+			case EpicRtcVideoCodec::VP9:
+			{
+				return TEXT("VP9");
+			}
+			default:
+			{
+				return TEXT("Unknown");
+			}
+		}
+	}
+
+	inline FString ToString(EpicRtcRoomState State)
+	{
+		switch (State)
+		{
+			case EpicRtcRoomState::New:
+				return TEXT("New");
+			case EpicRtcRoomState::Pending:
+				return TEXT("Pending");
+			case EpicRtcRoomState::Joined:
+				return TEXT("Joined");
+			case EpicRtcRoomState::Left:
+				return TEXT("Left");
+			case EpicRtcRoomState::Failed:
+				return TEXT("Failed");
+			case EpicRtcRoomState::Exiting:
+				return TEXT("Exiting");
+		}
+		return TEXT("Unknown");
+	}
+
+	inline FString ToString(const TSharedPtr<FJsonObject>& JsonObj, bool bPretty = true)
+	{
+		FString Res;
+		if (bPretty)
+		{
+			auto JsonWriter = TJsonWriterFactory<TCHAR, TPrettyJsonPrintPolicy<TCHAR>>::Create(&Res);
+			FJsonSerializer::Serialize(JsonObj.ToSharedRef(), JsonWriter);
+		}
+		else
+		{
+			auto JsonWriter = TJsonWriterFactory<TCHAR, TCondensedJsonPrintPolicy<TCHAR>>::Create(&Res);
+			FJsonSerializer::Serialize(JsonObj.ToSharedRef(), JsonWriter);
+		}
+		return Res;
+	}
+
+	inline FString ToString(const EpicRtcStringView& Str)
+	{
+		FUtf8String Utf8String = FUtf8String::ConstructFromPtrSize(Str._ptr, Str._length);
+		return FString(Utf8String);
+	}
+
+	inline EpicRtcStringView ToEpicRtcStringView(const FUtf8String& Str)
+	{
+		return EpicRtcStringView{ ._ptr = (const char*)*Str, ._length = static_cast<uint64>(Str.Len()) };
+	}
+
+	/**
+	 * Reads a string represented by 2 bytes length (in bytes) followed by UTF16 characters.
+	 * String and length are encoded in little endian format.
+	 */
+	inline FString ReadString(const uint8*& Data, uint32_t& Size)
+	{
+		uint16_t BytesLength = Data[1] << 8 | Data[0];
+		check(Size >= (uint32_t)(BytesLength + 2));
+		Data += 2;
+		Size -= 2;
+
+		FString Message(BytesLength / sizeof(TCHAR), reinterpret_cast<const TCHAR*>(Data));
+		Data += BytesLength;
+		Size -= BytesLength;
+
+		return Message;
+	}
+} // namespace UE::PixelStreaming2
diff --git a/Engine/Plugins/Media/PixelStreaming2/Source/PixelStreaming2/Private/UtilsVideo.h b/Engine/Plugins/Media/PixelStreaming2/Source/PixelStreaming2/Private/UtilsVideo.h
new file mode 100644
--- /dev/null
+++ b/Engine/Plugins/Media/PixelStreaming2/Source/PixelStreaming2/Private/UtilsVideo.h
@@ -0,0 +1,254 @@
+// Copyright Epic Games, Inc. All Rights Reserved.
+
+#pragma once
+
+#include "Containers/Array.h"
+#include "EpicRtcVideoCommon.h"
+#include "IPixelStreaming2Streamer.h"
+#include "PixelStreaming2PluginSettings.h"
+#include "Video/VideoEncoder.h"
+
+#include "epic_rtc/core/video/video_common.h"
+#include "epic_rtc/core/video/video_rate_control.h"
+
+namespace UE::PixelStreaming2
+{
+	// HACK (aidan.possemiers) the AVCodecs API surface wants a SharedPtr for the encoded data but EpicRtc already owns that and we don't want AVCodecs to delete it
+	struct FFakeDeleter
+	{
+		void operator()(uint8* Object) const
+		{
+		}
+	};
+
+	// List of video codecs supported by PS2
+	const TStaticArray<EVideoCodec, 4> SupportedVideoCodecs = { EVideoCodec::H264, EVideoCodec::AV1, EVideoCodec::VP8, EVideoCodec::VP9 };
+
+	constexpr uint32_t NumSimulcastLayers = 3;
+	// Each subsequent layer is 1/ScalingFactor the size of the previous
+	constexpr uint32_t ScalingFactor = 2;
+
+	// Helper array for all scalability modes. EScalabilityMode::None must always be the last entry
+	const TArray<EScalabilityMode> AllScalabilityModes = {
+		EScalabilityMode::L1T1,
+		EScalabilityMode::L1T2,
+		EScalabilityMode::L1T3,
+		EScalabilityMode::L2T1,
+		EScalabilityMode::L2T1h,
+		EScalabilityMode::L2T1_KEY,
+		EScalabilityMode::L2T2,
+		EScalabilityMode::L2T2h,
+		EScalabilityMode::L2T2_KEY,
+		EScalabilityMode::L2T2_KEY_SHIFT,
+		EScalabilityMode::L2T3,
+		EScalabilityMode::L2T3h,
+		EScalabilityMode::L2T3_KEY,
+		EScalabilityMode::L3T1,
+		EScalabilityMode::L3T1h,
+		EScalabilityMode::L3T1_KEY,
+		EScalabilityMode::L3T2,
+		EScalabilityMode::L3T2h,
+		EScalabilityMode::L3T2_KEY,
+		EScalabilityMode::L3T3,
+		EScalabilityMode::L3T3h,
+		EScalabilityMode::L3T3_KEY,
+		EScalabilityMode::S2T1,
+		EScalabilityMode::S2T1h,
+		EScalabilityMode::S2T2,
+		EScalabilityMode::S2T2h,
+		EScalabilityMode::S2T3,
+		EScalabilityMode::S2T3h,
+		EScalabilityMode::S3T1,
+		EScalabilityMode::S3T1h,
+		EScalabilityMode::S3T2,
+		EScalabilityMode::S3T2h,
+		EScalabilityMode::S3T3,
+		EScalabilityMode::S3T3h,
+		EScalabilityMode::None
+	};
+
+	// Make sure EpicRtcVideoScalabilityMode and EScalabilityMode match up
+	static_assert(EpicRtcVideoScalabilityMode::L1T1 == static_cast<EpicRtcVideoScalabilityMode>(EScalabilityMode::L1T1));
+	static_assert(EpicRtcVideoScalabilityMode::L1T2 == static_cast<EpicRtcVideoScalabilityMode>(EScalabilityMode::L1T2));
+	static_assert(EpicRtcVideoScalabilityMode::L1T3 == static_cast<EpicRtcVideoScalabilityMode>(EScalabilityMode::L1T3));
+	static_assert(EpicRtcVideoScalabilityMode::L2T1 == static_cast<EpicRtcVideoScalabilityMode>(EScalabilityMode::L2T1));
+	static_assert(EpicRtcVideoScalabilityMode::L2T1h == static_cast<EpicRtcVideoScalabilityMode>(EScalabilityMode::L2T1h));
+	static_assert(EpicRtcVideoScalabilityMode::L2T1Key == static_cast<EpicRtcVideoScalabilityMode>(EScalabilityMode::L2T1_KEY));
+	static_assert(EpicRtcVideoScalabilityMode::L2T2 == static_cast<EpicRtcVideoScalabilityMode>(EScalabilityMode::L2T2));
+	static_assert(EpicRtcVideoScalabilityMode::L2T2h == static_cast<EpicRtcVideoScalabilityMode>(EScalabilityMode::L2T2h));
+	static_assert(EpicRtcVideoScalabilityMode::L2T2Key == static_cast<EpicRtcVideoScalabilityMode>(EScalabilityMode::L2T2_KEY));
+	static_assert(EpicRtcVideoScalabilityMode::L2T2KeyShift == static_cast<EpicRtcVideoScalabilityMode>(EScalabilityMode::L2T2_KEY_SHIFT));
+	static_assert(EpicRtcVideoScalabilityMode::L2T3 == static_cast<EpicRtcVideoScalabilityMode>(EScalabilityMode::L2T3));
+	static_assert(EpicRtcVideoScalabilityMode::L2T3h == static_cast<EpicRtcVideoScalabilityMode>(EScalabilityMode::L2T3h));
+	static_assert(EpicRtcVideoScalabilityMode::L2T3Key == static_cast<EpicRtcVideoScalabilityMode>(EScalabilityMode::L2T3_KEY));
+	static_assert(EpicRtcVideoScalabilityMode::L3T1 == static_cast<EpicRtcVideoScalabilityMode>(EScalabilityMode::L3T1));
+	static_assert(EpicRtcVideoScalabilityMode::L3T1h == static_cast<EpicRtcVideoScalabilityMode>(EScalabilityMode::L3T1h));
+	static_assert(EpicRtcVideoScalabilityMode::L3T1Key == static_cast<EpicRtcVideoScalabilityMode>(EScalabilityMode::L3T1_KEY));
+	static_assert(EpicRtcVideoScalabilityMode::L3T2 == static_cast<EpicRtcVideoScalabilityMode>(EScalabilityMode::L3T2));
+	static_assert(EpicRtcVideoScalabilityMode::L3T2h == static_cast<EpicRtcVideoScalabilityMode>(EScalabilityMode::L3T2h));
+	static_assert(EpicRtcVideoScalabilityMode::L3T2Key == static_cast<EpicRtcVideoScalabilityMode>(EScalabilityMode::L3T2_KEY));
+	static_assert(EpicRtcVideoScalabilityMode::L3T3 == static_cast<EpicRtcVideoScalabilityMode>(EScalabilityMode::L3T3));
+	static_assert(EpicRtcVideoScalabilityMode::L3T3h == static_cast<EpicRtcVideoScalabilityMode>(EScalabilityMode::L3T3h));
+	static_assert(EpicRtcVideoScalabilityMode::L3T3Key == static_cast<EpicRtcVideoScalabilityMode>(EScalabilityMode::L3T3_KEY));
+	static_assert(EpicRtcVideoScalabilityMode::S2T1 == static_cast<EpicRtcVideoScalabilityMode>(EScalabilityMode::S2T1));
+	static_assert(EpicRtcVideoScalabilityMode::S2T1h == static_cast<EpicRtcVideoScalabilityMode>(EScalabilityMode::S2T1h));
+	static_assert(EpicRtcVideoScalabilityMode::S2T2 == static_cast<EpicRtcVideoScalabilityMode>(EScalabilityMode::S2T2));
+	static_assert(EpicRtcVideoScalabilityMode::S2T2h == static_cast<EpicRtcVideoScalabilityMode>(EScalabilityMode::S2T2h));
+	static_assert(EpicRtcVideoScalabilityMode::S2T3 == static_cast<EpicRtcVideoScalabilityMode>(EScalabilityMode::S2T3));
+	static_assert(EpicRtcVideoScalabilityMode::S2T3h == static_cast<EpicRtcVideoScalabilityMode>(EScalabilityMode::S2T3h));
+	static_assert(EpicRtcVideoScalabilityMode::S3T1 == static_cast<EpicRtcVideoScalabilityMode>(EScalabilityMode::S3T1));
+	static_assert(EpicRtcVideoScalabilityMode::S3T1h == static_cast<EpicRtcVideoScalabilityMode>(EScalabilityMode::S3T1h));
+	static_assert(EpicRtcVideoScalabilityMode::S3T2 == static_cast<EpicRtcVideoScalabilityMode>(EScalabilityMode::S3T2));
+	static_assert(EpicRtcVideoScalabilityMode::S3T2h == static_cast<EpicRtcVideoScalabilityMode>(EScalabilityMode::S3T2h));
+	static_assert(EpicRtcVideoScalabilityMode::S3T3 == static_cast<EpicRtcVideoScalabilityMode>(EScalabilityMode::S3T3));
+	static_assert(EpicRtcVideoScalabilityMode::S3T3h == static_cast<EpicRtcVideoScalabilityMode>(EScalabilityMode::S3T3h));
+	static_assert(EpicRtcVideoScalabilityMode::None == static_cast<EpicRtcVideoScalabilityMode>(EScalabilityMode::None));
+
+	/**
+	 * A struct representing the simulcast paramaters of a single simulcast layer used by PixelStreaming2.
+	 * Specifically, each layer has a `Scaling`, `MinBitrate` and `MaxBitrate`.
+	 */
+	struct FPixelStreaming2SimulcastLayer
+	{
+		float Scaling;
+		int	  MinBitrate;
+		int	  MaxBitrate;
+	};
+
+	inline TArray<FPixelStreaming2SimulcastLayer> GetSimulcastParameters()
+	{
+		TArray<FPixelStreaming2SimulcastLayer> SimulcastParams;
+
+		if (UPixelStreaming2PluginSettings::CVarEncoderEnableSimulcast.GetValueOnAnyThread())
+		{
+			int MinBps = UPixelStreaming2PluginSettings::CVarWebRTCMinBitrate.GetValueOnAnyThread();
+			int MaxBps = UPixelStreaming2PluginSettings::CVarWebRTCMaxBitrate.GetValueOnAnyThread();
+			int BpsRange = MaxBps - MinBps;
+
+			// Bitrates assignment per layer. Effectively
+			// 0: 0 -> 1/4
+			// 1: 1/4 -> 1/2
+			// 2: 1/2 -> 1
+			TArray<TTuple<int, int>> Bitrates = {
+				{ MinBps, MinBps + (BpsRange / 4) },
+				{ MinBps + (BpsRange / 4), MinBps + (BpsRange / 2) },
+				{ MinBps + (BpsRange / 2), MaxBps }
+			};
+
+			for (int i = 0; i < NumSimulcastLayers; i++)
+			{
+				// EpicRtc expects the layers to be added in order of scaling factors from largest to smallest (ie smallest res to largest res)
+				float Scaling = ScalingFactor * (NumSimulcastLayers - i - 1);
+
+				// clang-format off
+				SimulcastParams.Add({
+					.Scaling = Scaling > 0 ? Scaling : 1.f,
+					.MinBitrate = Bitrates[i].Get<0>(),
+					.MaxBitrate = Bitrates[i].Get<1>(),
+				});
+				// clang-format on
+			}
+		}
+		else
+		{
+			// clang-format off
+			SimulcastParams.Add({ 
+				.Scaling = 1.f,
+				.MinBitrate = UPixelStreaming2PluginSettings::CVarWebRTCMinBitrate.GetValueOnAnyThread(),
+				.MaxBitrate = UPixelStreaming2PluginSettings::CVarWebRTCMaxBitrate.GetValueOnAnyThread() 
+			});
+			// clang-format on
+		}
+
+		return SimulcastParams;
+	}
+
+	inline FEpicRtcParameterPairArray* CreateH264Format(EH264Profile Profile, UE::AVCodecCore::H264::EH264Level Level)
+	{
+		// TODO (Migration): RTCP-7028 picRtcStringView needs a way to own the memory passed into it
+		// return new FEpicRtcParameterPairArray(
+		// {
+		// 	EpicRtcParameterPair{
+		// 		._key = EpicRtcStringView{ ._ptr = "profile-level-id", ._length = 16 },
+		// 		._value = EpicRtcStringView{ ._ptr = TCHAR_TO_ANSI(*ProfileString), ._length = (uint64_t)ProfileString.Len() }
+		// 	},
+		// 	EpicRtcParameterPair{
+		// 		._key = EpicRtcStringView{ ._ptr = "packetization-mode", ._length = 18 },
+		// 		._value = EpicRtcStringView{ ._ptr = "1", ._length = 1 }
+		// 	},
+		// 	EpicRtcParameterPair{
+		// 		._key = EpicRtcStringView{ ._ptr = "level-asymmetry-allowed", ._length = 23 },
+		// 		._value = EpicRtcStringView{ ._ptr = "1", ._length = 1 }
+		// 	}
+		// });
+
+		using namespace UE::AVCodecCore::H264;
+		if (Profile == EH264Profile::ConstrainedBaseline && Level == EH264Level::Level_3_1)
+		{
+			return new FEpicRtcParameterPairArray(
+				{ EpicRtcParameterPair{
+					  ._key = EpicRtcStringView{ ._ptr = "profile-level-id", ._length = 16 },
+					  ._value = EpicRtcStringView{ ._ptr = "42e01f", ._length = 6 } },
+					EpicRtcParameterPair{
+						._key = EpicRtcStringView{ ._ptr = "packetization-mode", ._length = 18 },
+						._value = EpicRtcStringView{ ._ptr = "1", ._length = 1 } },
+					EpicRtcParameterPair{
+						._key = EpicRtcStringView{ ._ptr = "level-asymmetry-allowed", ._length = 23 },
+						._value = EpicRtcStringView{ ._ptr = "1", ._length = 1 } } });
+		}
+		else if (Profile == EH264Profile::Baseline && Level == EH264Level::Level_3_1)
+		{
+			return new FEpicRtcParameterPairArray(
+				{ EpicRtcParameterPair{
+					  ._key = EpicRtcStringView{ ._ptr = "profile-level-id", ._length = 16 },
+					  ._value = EpicRtcStringView{ ._ptr = "42001f", ._length = 6 } },
+					EpicRtcParameterPair{
+						._key = EpicRtcStringView{ ._ptr = "packetization-mode", ._length = 18 },
+						._value = EpicRtcStringView{ ._ptr = "1", ._length = 1 } },
+					EpicRtcParameterPair{
+						._key = EpicRtcStringView{ ._ptr = "level-asymmetry-allowed", ._length = 23 },
+						._value = EpicRtcStringView{ ._ptr = "1", ._length = 1 } } });
+		}
+		else
+		{
+			return nullptr;
+		}
+	}
+
+	inline FEpicRtcParameterPairArray* CreateVP9Format(UE::AVCodecCore::VP9::EProfile Profile)
+	{
+		using namespace UE::AVCodecCore::VP9;
+		if (Profile == EProfile::Profile0)
+		{
+			return new FEpicRtcParameterPairArray(
+				{ EpicRtcParameterPair{
+					._key = EpicRtcStringView{ ._ptr = "profile-id", ._length = 10 },
+					._value = EpicRtcStringView{ ._ptr = "0", ._length = 1 } } });
+		}
+		else if (Profile == EProfile::Profile1)
+		{
+			return new FEpicRtcParameterPairArray(
+				{ EpicRtcParameterPair{
+					._key = EpicRtcStringView{ ._ptr = "profile-id", ._length = 10 },
+					._value = EpicRtcStringView{ ._ptr = "1", ._length = 1 } } });
+		}
+		else if (Profile == EProfile::Profile2)
+		{
+			return new FEpicRtcParameterPairArray(
+				{ EpicRtcParameterPair{
+					._key = EpicRtcStringView{ ._ptr = "profile-id", ._length = 10 },
+					._value = EpicRtcStringView{ ._ptr = "2", ._length = 1 } } });
+		}
+		else if (Profile == EProfile::Profile3)
+		{
+			return new FEpicRtcParameterPairArray(
+				{ EpicRtcParameterPair{
+					._key = EpicRtcStringView{ ._ptr = "profile-id", ._length = 10 },
+					._value = EpicRtcStringView{ ._ptr = "3", ._length = 1 } } });
+		}
+		else
+		{
+			return nullptr;
+		}
+	}
+} // namespace UE::PixelStreaming2
\ No newline at end of file
diff --git a/Engine/Plugins/Media/PixelStreaming2/Source/PixelStreaming2/Private/VideoCapturer.cpp b/Engine/Plugins/Media/PixelStreaming2/Source/PixelStreaming2/Private/VideoCapturer.cpp
--- a/Engine/Plugins/Media/PixelStreaming2/Source/PixelStreaming2/Private/VideoCapturer.cpp
+++ b/Engine/Plugins/Media/PixelStreaming2/Source/PixelStreaming2/Private/VideoCapturer.cpp
@@ -13,7 +13,7 @@
 #include "PixelStreaming2PluginSettings.h"
 #include "Logging.h"
 #include "PixelStreaming2Trace.h"
-#include "VideoUtils.h"
+#include "UtilsVideo.h"
 
 namespace UE::PixelStreaming2
 {
diff --git a/Engine/Plugins/Media/PixelStreaming2/Source/PixelStreaming2/Private/VideoProducerBackBuffer.cpp b/Engine/Plugins/Media/PixelStreaming2/Source/PixelStreaming2/Private/VideoProducerBackBuffer.cpp
--- a/Engine/Plugins/Media/PixelStreaming2/Source/PixelStreaming2/Private/VideoProducerBackBuffer.cpp
+++ b/Engine/Plugins/Media/PixelStreaming2/Source/PixelStreaming2/Private/VideoProducerBackBuffer.cpp
@@ -6,7 +6,10 @@
 #include "Logging.h"
 #include "PixelCaptureBufferFormat.h"
 #include "PixelCaptureInputFrameRHI.h"
-#include "Utils.h"
+#include "RHI.h"
+#include "UtilsCommon.h"
+#include "Widgets/SWindow.h"
+
 
 namespace UE::PixelStreaming2
 {
diff --git a/Engine/Plugins/Media/PixelStreaming2/Source/PixelStreaming2/Private/VideoProducerPIEViewport.cpp b/Engine/Plugins/Media/PixelStreaming2/Source/PixelStreaming2/Private/VideoProducerPIEViewport.cpp
--- a/Engine/Plugins/Media/PixelStreaming2/Source/PixelStreaming2/Private/VideoProducerPIEViewport.cpp
+++ b/Engine/Plugins/Media/PixelStreaming2/Source/PixelStreaming2/Private/VideoProducerPIEViewport.cpp
@@ -5,8 +5,8 @@
 #include "Engine/GameViewportClient.h"
 #include "PixelCaptureInputFrameRHI.h"
 #include "RenderingThread.h"
-#include "Utils.h"
 #include "UnrealClient.h"
+#include "UtilsCommon.h"
 
 namespace UE::PixelStreaming2
 {
diff --git a/Engine/Plugins/Media/PixelStreaming2/Source/PixelStreaming2/Private/VideoProducerRenderTarget.cpp b/Engine/Plugins/Media/PixelStreaming2/Source/PixelStreaming2/Private/VideoProducerRenderTarget.cpp
--- a/Engine/Plugins/Media/PixelStreaming2/Source/PixelStreaming2/Private/VideoProducerRenderTarget.cpp
+++ b/Engine/Plugins/Media/PixelStreaming2/Source/PixelStreaming2/Private/VideoProducerRenderTarget.cpp
@@ -4,7 +4,9 @@
 
 #include "Engine/TextureRenderTarget2D.h"
 #include "PixelCaptureInputFrameRHI.h"
+#include "RHI.h"
 #include "TextureResource.h"
+#include "Widgets/SWindow.h"
 
 namespace UE::PixelStreaming2
 {
diff --git a/Engine/Plugins/Media/PixelStreaming2/Source/PixelStreaming2/Private/VideoSourceGroup.cpp b/Engine/Plugins/Media/PixelStreaming2/Source/PixelStreaming2/Private/VideoSourceGroup.cpp
--- a/Engine/Plugins/Media/PixelStreaming2/Source/PixelStreaming2/Private/VideoSourceGroup.cpp
+++ b/Engine/Plugins/Media/PixelStreaming2/Source/PixelStreaming2/Private/VideoSourceGroup.cpp
@@ -106,7 +106,7 @@ namespace UE::PixelStreaming2
 			return this;
 		}
 
-		/*
+		/**
 		 * Note this function is required as part of `FSingleThreadRunnable` and only gets called when engine is run in single-threaded mode,
 		 * so the logic is much less complex as this is not a case we particularly optimize for, a simple tick on an interval will be acceptable.
 		 */
@@ -135,7 +135,7 @@ namespace UE::PixelStreaming2
 		TWeakPtr<FVideoSourceGroup> OuterVideoSourceGroup = nullptr;
 		uint64						LastSubmitCycles = 0;
 
-		/* Use this event to signal when we should wake and also how long we should sleep for between transmitting a frame. */
+		// Use this event to signal when we should wake and also how long we should sleep for between transmitting a frame.
 		FEventRef FrameEvent;
 	};
 
diff --git a/Engine/Plugins/Media/PixelStreaming2/Source/PixelStreaming2/Private/VideoUtils.h b/Engine/Plugins/Media/PixelStreaming2/Source/PixelStreaming2/Private/VideoUtils.h
deleted file mode 100644
diff --git a/Engine/Plugins/Media/PixelStreaming2/Source/PixelStreaming2/Public/IPixelStreaming2AudioConsumer.h b/Engine/Plugins/Media/PixelStreaming2/Source/PixelStreaming2/Public/IPixelStreaming2AudioConsumer.h
--- a/Engine/Plugins/Media/PixelStreaming2/Source/PixelStreaming2/Public/IPixelStreaming2AudioConsumer.h
+++ b/Engine/Plugins/Media/PixelStreaming2/Source/PixelStreaming2/Public/IPixelStreaming2AudioConsumer.h
@@ -4,7 +4,7 @@
 
 #include "CoreMinimal.h"
 
-/*
+/**
  * An "Audio Consumer" is an object that is responsible for outputting the audio received from a peer. For example, by
  * passing the audio into a UE submix.
  */
diff --git a/Engine/Plugins/Media/PixelStreaming2/Source/PixelStreaming2/Public/IPixelStreaming2AudioProducer.h b/Engine/Plugins/Media/PixelStreaming2/Source/PixelStreaming2/Public/IPixelStreaming2AudioProducer.h
--- a/Engine/Plugins/Media/PixelStreaming2/Source/PixelStreaming2/Public/IPixelStreaming2AudioProducer.h
+++ b/Engine/Plugins/Media/PixelStreaming2/Source/PixelStreaming2/Public/IPixelStreaming2AudioProducer.h
@@ -4,7 +4,7 @@
 
 #include "CoreTypes.h"
 
-/*
+/**
  * A "Audio Producer" is an object that you use to push audio frames into the Pixel Streaming system.
  *
  * Example usage:
@@ -26,5 +26,5 @@ public:
 	 * @param NumChannels Number of audio channels. For example 2 for stero audio.
 	 * @param SampleRate Audio sample rate in samples per second.
 	 */
-    virtual void PushAudio(const float* InBuffer, int32 NumSamples, int32 NumChannels, int32 SampleRate) = 0;
+	virtual void PushAudio(const float* InBuffer, int32 NumSamples, int32 NumChannels, int32 SampleRate) = 0;
 };
diff --git a/Engine/Plugins/Media/PixelStreaming2/Source/PixelStreaming2/Public/IPixelStreaming2AudioSink.h b/Engine/Plugins/Media/PixelStreaming2/Source/PixelStreaming2/Public/IPixelStreaming2AudioSink.h
--- a/Engine/Plugins/Media/PixelStreaming2/Source/PixelStreaming2/Public/IPixelStreaming2AudioSink.h
+++ b/Engine/Plugins/Media/PixelStreaming2/Source/PixelStreaming2/Public/IPixelStreaming2AudioSink.h
@@ -6,7 +6,7 @@
 
 class IPixelStreaming2AudioConsumer;
 
-/*
+/**
  * An "Audio Sink" is an object that receives audio from a singular peer. From here, you can add multiple consumers to
  * output the received audio.
  */
diff --git a/Engine/Plugins/Media/PixelStreaming2/Source/PixelStreaming2/Public/IPixelStreaming2Module.h b/Engine/Plugins/Media/PixelStreaming2/Source/PixelStreaming2/Public/IPixelStreaming2Module.h
--- a/Engine/Plugins/Media/PixelStreaming2/Source/PixelStreaming2/Public/IPixelStreaming2Module.h
+++ b/Engine/Plugins/Media/PixelStreaming2/Source/PixelStreaming2/Public/IPixelStreaming2Module.h
@@ -41,7 +41,7 @@ public:
 		return FModuleManager::Get().IsModuleLoaded("PixelStreaming2");
 	}
 
-	/*
+	/**
 	 * Event fired when internal streamer is initialized and the methods on this module are ready for use.
 	 */
 	DECLARE_EVENT_OneParam(IPixelStreaming2Module, FReadyEvent, IPixelStreaming2Module&);
diff --git a/Engine/Plugins/Media/PixelStreaming2/Source/PixelStreaming2/Public/IPixelStreaming2VideoConsumer.h b/Engine/Plugins/Media/PixelStreaming2/Source/PixelStreaming2/Public/IPixelStreaming2VideoConsumer.h
--- a/Engine/Plugins/Media/PixelStreaming2/Source/PixelStreaming2/Public/IPixelStreaming2VideoConsumer.h
+++ b/Engine/Plugins/Media/PixelStreaming2/Source/PixelStreaming2/Public/IPixelStreaming2VideoConsumer.h
@@ -5,7 +5,7 @@
 #include "CoreMinimal.h"
 #include "Rendering/SlateRenderer.h"
 
-/*
+/**
  * An "Video Consumer" is an object that is responsible for outputting the video received from a peer. For example, by
  * rendering to a render target.
  */
diff --git a/Engine/Plugins/Media/PixelStreaming2/Source/PixelStreaming2/Public/IPixelStreaming2VideoProducer.h b/Engine/Plugins/Media/PixelStreaming2/Source/PixelStreaming2/Public/IPixelStreaming2VideoProducer.h
--- a/Engine/Plugins/Media/PixelStreaming2/Source/PixelStreaming2/Public/IPixelStreaming2VideoProducer.h
+++ b/Engine/Plugins/Media/PixelStreaming2/Source/PixelStreaming2/Public/IPixelStreaming2VideoProducer.h
@@ -5,7 +5,7 @@
 #include "CoreTypes.h"
 #include "IPixelCaptureInputFrame.h"
 
-/*
+/**
  * A "Video Producer" is an object that you use to push video frames into the Pixel Streaming system.
  *
  * Example usage:
diff --git a/Engine/Plugins/Media/PixelStreaming2/Source/PixelStreaming2/Public/IPixelStreaming2VideoSink.h b/Engine/Plugins/Media/PixelStreaming2/Source/PixelStreaming2/Public/IPixelStreaming2VideoSink.h
--- a/Engine/Plugins/Media/PixelStreaming2/Source/PixelStreaming2/Public/IPixelStreaming2VideoSink.h
+++ b/Engine/Plugins/Media/PixelStreaming2/Source/PixelStreaming2/Public/IPixelStreaming2VideoSink.h
@@ -6,7 +6,7 @@
 
 class IPixelStreaming2VideoConsumer;
 
-/*
+/**
  * Interface for a sink that collects video coming in from the browser and passes into into UE.
  */
 class PIXELSTREAMING2_API IPixelStreaming2VideoSink
diff --git a/Engine/Plugins/Media/PixelStreaming2/Source/PixelStreaming2Core/Internal/PixelStreaming2PluginSettings.h b/Engine/Plugins/Media/PixelStreaming2/Source/PixelStreaming2Core/Internal/PixelStreaming2PluginSettings.h
--- a/Engine/Plugins/Media/PixelStreaming2/Source/PixelStreaming2Core/Internal/PixelStreaming2PluginSettings.h
+++ b/Engine/Plugins/Media/PixelStreaming2/Source/PixelStreaming2Core/Internal/PixelStreaming2PluginSettings.h
@@ -54,7 +54,7 @@ enum class EPortAllocatorFlags : uint32
 };
 ENUM_CLASS_FLAGS(EPortAllocatorFlags);
 
-/* Pixel Streaming can limit who can send input (keyboard, mouse, etc). */
+/** Pixel Streaming can limit who can send input (keyboard, mouse, etc). */
 UENUM()
 enum class EInputControllerMode : uint8
 {
diff --git a/Engine/Plugins/Media/PixelStreaming2/Source/PixelStreaming2Editor/Internal/VideoProducerBackBufferComposited.h b/Engine/Plugins/Media/PixelStreaming2/Source/PixelStreaming2Editor/Internal/VideoProducerBackBufferComposited.h
--- a/Engine/Plugins/Media/PixelStreaming2/Source/PixelStreaming2Editor/Internal/VideoProducerBackBufferComposited.h
+++ b/Engine/Plugins/Media/PixelStreaming2/Source/PixelStreaming2Editor/Internal/VideoProducerBackBufferComposited.h
@@ -12,7 +12,7 @@ namespace UE::EditorPixelStreaming2
 {
 	using namespace UE::PixelStreaming2;
 
-	/*
+	/**
 	 * Use this if you want to send the full UE editor as video input.
 	 */
 	class PIXELSTREAMING2EDITOR_API FVideoProducerBackBufferComposited : public FVideoProducer
diff --git a/Engine/Plugins/Media/PixelStreaming2/Source/PixelStreaming2Editor/Internal/VideoProducerViewport.h b/Engine/Plugins/Media/PixelStreaming2/Source/PixelStreaming2Editor/Internal/VideoProducerViewport.h
--- a/Engine/Plugins/Media/PixelStreaming2/Source/PixelStreaming2Editor/Internal/VideoProducerViewport.h
+++ b/Engine/Plugins/Media/PixelStreaming2/Source/PixelStreaming2Editor/Internal/VideoProducerViewport.h
@@ -12,7 +12,7 @@ namespace UE::EditorPixelStreaming2
 
 	using namespace UE::PixelStreaming2;
 
-	/*
+	/**
 	 * Use this if you want to send the UE primary scene viewport as video input - will only work in editor.
 	 */
 	class PIXELSTREAMING2EDITOR_API FVideoProducerViewport : public FVideoProducer
diff --git a/Engine/Plugins/Media/PixelStreaming2/Source/PixelStreaming2Editor/Private/Factories.cpp b/Engine/Plugins/Media/PixelStreaming2/Source/PixelStreaming2Editor/Private/Factories.cpp
--- a/Engine/Plugins/Media/PixelStreaming2/Source/PixelStreaming2Editor/Private/Factories.cpp
+++ b/Engine/Plugins/Media/PixelStreaming2/Source/PixelStreaming2Editor/Private/Factories.cpp
@@ -11,7 +11,7 @@
 
 #define LOCTEXT_NAMESPACE "PixelStreaming2"
 
-/*
+/**
  * ---------- UPixelStreaming2MediaTextureFactory -------------------
  */
 UPixelStreaming2MediaTextureFactory::UPixelStreaming2MediaTextureFactory(const FObjectInitializer& ObjectInitializer)
@@ -44,7 +44,7 @@ UObject* UPixelStreaming2MediaTextureFactory::FactoryCreateNew(UClass* InClass,
 	return nullptr;
 }
 
-/*
+/**
  * ---------- UPixelStreaming2VideoProducerBackBufferFactory -------------------
  */
 UPixelStreaming2VideoProducerBackBufferFactory::UPixelStreaming2VideoProducerBackBufferFactory(const FObjectInitializer& ObjectInitializer)
@@ -76,7 +76,7 @@ UObject* UPixelStreaming2VideoProducerBackBufferFactory::FactoryCreateNew(UClass
 	return nullptr;
 }
 
-/*
+/**
  * ---------- UPixelStreaming2VideoProducerMediaCaptureFactory -------------------
  */
 UPixelStreaming2VideoProducerMediaCaptureFactory::UPixelStreaming2VideoProducerMediaCaptureFactory(const FObjectInitializer& ObjectInitializer)
@@ -108,7 +108,7 @@ UObject* UPixelStreaming2VideoProducerMediaCaptureFactory::FactoryCreateNew(UCla
 	return nullptr;
 }
 
-/*
+/**
  * ---------- UPixelStreaming2VideoProducerRenderTargetFactory -------------------
  */
 UPixelStreaming2VideoProducerRenderTargetFactory::UPixelStreaming2VideoProducerRenderTargetFactory(const FObjectInitializer& ObjectInitializer)
diff --git a/Engine/Plugins/Media/PixelStreaming2/Source/PixelStreaming2Editor/Private/Factories.h b/Engine/Plugins/Media/PixelStreaming2/Source/PixelStreaming2Editor/Private/Factories.h
--- a/Engine/Plugins/Media/PixelStreaming2/Source/PixelStreaming2Editor/Private/Factories.h
+++ b/Engine/Plugins/Media/PixelStreaming2/Source/PixelStreaming2Editor/Private/Factories.h
@@ -8,7 +8,7 @@
 
 #include "Factories.generated.h"
 
-/*
+/**
  * ---------- UPixelStreaming2MediaTextureFactory -------------------
  */
 UCLASS()
@@ -22,7 +22,7 @@ public:
 	virtual UObject* FactoryCreateNew(UClass* InClass, UObject* InParent, FName InName, EObjectFlags Flags, UObject* Context, FFeedbackContext* Warn) override;
 };
 
-/*
+/**
  * ---------- UPixelStreaming2VideoProducerBackBufferFactory -------------------
  */
 UCLASS()
@@ -36,7 +36,7 @@ public:
 	virtual UObject* FactoryCreateNew(UClass* InClass, UObject* InParent, FName InName, EObjectFlags Flags, UObject* Context, FFeedbackContext* Warn) override;
 };
 
-/*
+/**
  * ---------- UPixelStreaming2VideoProducerMediaCaptureFactory -------------------
  */
 UCLASS()
@@ -50,7 +50,7 @@ public:
 	virtual UObject* FactoryCreateNew(UClass* InClass, UObject* InParent, FName InName, EObjectFlags Flags, UObject* Context, FFeedbackContext* Warn) override;
 };
 
-/*
+/**
  * ---------- UPixelStreaming2VideoProducerRenderTargetFactory -------------------
  */
 UCLASS()
diff --git a/Engine/Plugins/Media/PixelStreaming2/Source/PixelStreaming2Editor/Private/PixelStreaming2EditorModule.cpp b/Engine/Plugins/Media/PixelStreaming2/Source/PixelStreaming2Editor/Private/PixelStreaming2EditorModule.cpp
--- a/Engine/Plugins/Media/PixelStreaming2/Source/PixelStreaming2Editor/Private/PixelStreaming2EditorModule.cpp
+++ b/Engine/Plugins/Media/PixelStreaming2/Source/PixelStreaming2Editor/Private/PixelStreaming2EditorModule.cpp
@@ -305,7 +305,7 @@ namespace UE::EditorPixelStreaming2
 		if (!DownloadProcess.IsValid())
 		{
 			// We set bSkipIfPresent to false, which means the get_ps_servers script will always be run, that script will choose whether to download or not
-			DownloadProcess = UE::PixelStreaming2Servers::DownloadPixelStreaming2Servers(/*bSkipIfPresent*/ false);
+			DownloadProcess = UE::PixelStreaming2Servers::DownloadPixelStreaming2Servers(false /* bSkipIfPresent */);
 			if (DownloadProcess.IsValid())
 			{
 				DownloadProcess->OnCompleted().BindLambda([this](int ExitCode) {
diff --git a/Engine/Plugins/Media/PixelStreaming2/Source/PixelStreaming2HMD/Private/PixelStreaming2HMD.h b/Engine/Plugins/Media/PixelStreaming2/Source/PixelStreaming2HMD/Private/PixelStreaming2HMD.h
--- a/Engine/Plugins/Media/PixelStreaming2/Source/PixelStreaming2HMD/Private/PixelStreaming2HMD.h
+++ b/Engine/Plugins/Media/PixelStreaming2/Source/PixelStreaming2HMD/Private/PixelStreaming2HMD.h
@@ -19,7 +19,7 @@ class UCanvas;
 class PIXELSTREAMING2HMD_API FPixelStreaming2HMD : public IPixelStreaming2HMD, public FHeadMountedDisplayBase, public FHMDSceneViewExtension
 {
 public:
-	/** IXRTrackingSystem interface */
+	// Begin IXRTrackingSystem
 	virtual FName GetSystemName() const override
 	{
 		static FName DefaultName(TEXT("PixelStreaming2HMD"));
@@ -53,13 +53,15 @@ public:
 	{
 		return SharedThis(this);
 	}
+	// End IXRTrackingSystem
 
 protected:
-	/** FXRTrackingSystemBase protected interface */
+	// Begin FXRTrackingSystemBase
 	virtual float GetWorldToMetersScale() const override;
+	// End FXRTrackingSystemBase
 
 public:
-	/** IHeadMountedDisplay interface */
+	// Begin IHeadMountedDisplay
 	virtual bool IsHMDConnected() override { return true; }
 	virtual bool IsHMDEnabled() const override;
 	virtual void EnableHMD(bool allow = true) override;
@@ -68,33 +70,36 @@ public:
 	virtual bool IsChromaAbCorrectionEnabled() const override;
 	virtual bool GetHMDDistortionEnabled(EShadingPath ShadingPath) const override { return false; }
 	virtual void DrawDistortionMesh_RenderThread(struct FHeadMountedDisplayPassContext& Context, const FIntPoint& TextureSize) override;
+	// End IHeadMountedDisplay
 
-	/** IStereoRendering interface */
+	// Begin IStereoRendering
 	virtual bool	IsStereoEnabled() const override;
 	virtual bool	EnableStereo(bool stereo = true) override;
 	virtual void	AdjustViewRect(int32 ViewIndex, int32& X, int32& Y, uint32& SizeX, uint32& SizeY) const override;
 	virtual void	CalculateStereoViewOffset(const int32 ViewIndex, FRotator& ViewRotation, const float InWorldToMeters, FVector& ViewLocation) override;
 	virtual FMatrix GetStereoProjectionMatrix(const int32 ViewIndex) const override;
 	virtual void	GetEyeRenderParams_RenderThread(const struct FHeadMountedDisplayPassContext& Context, FVector2D& EyeToSrcUVScaleValue, FVector2D& EyeToSrcUVOffsetValue) const override;
+	// End IStereoRendering
 
-	/** ISceneViewExtension interface */
+	// Begin ISceneViewExtension
 	virtual void SetupViewFamily(FSceneViewFamily& InViewFamily) override;
 	virtual void SetupView(FSceneViewFamily& InViewFamily, FSceneView& InView) override {};
 	virtual void BeginRenderViewFamily(FSceneViewFamily& InViewFamily) {}
 	virtual void PreRenderView_RenderThread(FRDGBuilder& GraphBuilder, FSceneView& InView) override;
 	virtual void PreRenderViewFamily_RenderThread(FRDGBuilder& GraphBuilder, FSceneViewFamily& InViewFamily) override;
+	// End ISceneViewExtension
 
-	/** Constructor */
 	FPixelStreaming2HMD(const FAutoRegister&);
 
-	/** Destructor */
 	virtual ~FPixelStreaming2HMD();
 
-	/** @return	True if the HMD was initialized OK */
+	/**
+	 * @return	True if the HMD was initialized OK
+	 */
 	bool IsInitialized() const { return true; }
 
 	virtual void SetTransform(FTransform Transform) override { CurHmdTransform = Transform; }
-	virtual void SetEyeViews(FTransform Left, FMatrix LeftProj, FTransform Right, FMatrix RightProj, FTransform HMD) override ;
+	virtual void SetEyeViews(FTransform Left, FMatrix LeftProj, FTransform Right, FMatrix RightProj, FTransform HMD) override;
 
 private:
 	FVector	   BasePosition = FVector::ZeroVector;
diff --git a/Engine/Plugins/Media/PixelStreaming2/Source/PixelStreaming2HMD/Private/PixelStreaming2HMDModule.h b/Engine/Plugins/Media/PixelStreaming2/Source/PixelStreaming2HMD/Private/PixelStreaming2HMDModule.h
--- a/Engine/Plugins/Media/PixelStreaming2/Source/PixelStreaming2HMD/Private/PixelStreaming2HMDModule.h
+++ b/Engine/Plugins/Media/PixelStreaming2/Source/PixelStreaming2HMD/Private/PixelStreaming2HMDModule.h
@@ -9,29 +9,29 @@ class IXRTrackingSystem;
 
 namespace UE::PixelStreaming2HMD
 {
-	/*
+	/**
 	 * This module allows HMD input to be used with pixel streaming
 	 */
 	class FPixelStreaming2HMDModule : public IPixelStreaming2HMDModule
 	{
 	public:
-		FPixelStreaming2HMD*	   GetPixelStreaming2HMD() const;
+		FPixelStreaming2HMD*	 GetPixelStreaming2HMD() const;
 		EPixelStreaming2XRSystem GetActiveXRSystem() { return ActiveXRSystem; }
-		void						   SetActiveXRSystem(EPixelStreaming2XRSystem System) { ActiveXRSystem = System; }
+		void					 SetActiveXRSystem(EPixelStreaming2XRSystem System) { ActiveXRSystem = System; }
 
 	private:
-		/** IModuleInterface implementation */
+		// Begin IModuleInterface
 		void StartupModule() override;
 		void ShutdownModule() override;
-		/** End IModuleInterface implementation */
+		// End IModuleInterface
 
-		/** IHeadMountedDisplayModule implementation */
+		// Begin IHeadMountedDisplayModule
 		virtual TSharedPtr<IXRTrackingSystem, ESPMode::ThreadSafe> CreateTrackingSystem() override;
 		FString													   GetModuleKeyName() const override { return FString(TEXT("PixelStreaming2HMD")); }
 		bool													   IsHMDConnected() override { return true; }
+		// End IHeadMountedDisplayModule
 
-		/** IHeadMountedDisplayModule implementation */
 		TSharedPtr<FPixelStreaming2HMD, ESPMode::ThreadSafe> HMD;
-		EPixelStreaming2XRSystem							   ActiveXRSystem;
+		EPixelStreaming2XRSystem							 ActiveXRSystem;
 	};
 } // namespace UE::PixelStreaming2HMD
\ No newline at end of file
diff --git a/Engine/Plugins/Media/PixelStreaming2/Source/PixelStreaming2Input/Internal/InputProtocolMap.h b/Engine/Plugins/Media/PixelStreaming2/Source/PixelStreaming2Input/Internal/InputProtocolMap.h
--- a/Engine/Plugins/Media/PixelStreaming2/Source/PixelStreaming2Input/Internal/InputProtocolMap.h
+++ b/Engine/Plugins/Media/PixelStreaming2/Source/PixelStreaming2Input/Internal/InputProtocolMap.h
@@ -22,17 +22,17 @@ namespace UE::PixelStreaming2Input
 		FInputProtocolMap(EPixelStreaming2MessageDirection InDirection)
 			: Direction(InDirection) {}
 
-		/* Begin IPixelStreaming2DataProtocol interface */
+		// Begin IPixelStreaming2DataProtocol interface
 		TSharedPtr<IPixelStreaming2InputMessage> Add(FString Key) override;
 		TSharedPtr<IPixelStreaming2InputMessage> Add(FString Key, TArray<EPixelStreaming2MessageTypes> InStructure) override;
 		TSharedPtr<IPixelStreaming2InputMessage> Find(FString Key) override;
-		FOnProtocolUpdated&							   OnProtocolUpdated() override { return OnProtocolUpdatedDelegate; };
-		TSharedPtr<FJsonObject>						   ToJson() override;
-		/* End IPixelStreaming2DataProtocol interface */
+		FOnProtocolUpdated&						 OnProtocolUpdated() override { return OnProtocolUpdatedDelegate; };
+		TSharedPtr<FJsonObject>					 ToJson() override;
+		// End IPixelStreaming2DataProtocol interface
 
-		bool												 AddInternal(FString Key, uint8 Id);
-		bool												 AddInternal(FString Key, uint8 Id, TArray<EPixelStreaming2MessageTypes> InStructure);
-		int													 Remove(FString Key);
+		bool										   AddInternal(FString Key, uint8 Id);
+		bool										   AddInternal(FString Key, uint8 Id, TArray<EPixelStreaming2MessageTypes> InStructure);
+		int											   Remove(FString Key);
 		const TSharedPtr<IPixelStreaming2InputMessage> Find(FString Key) const;
 
 		void Clear();
@@ -44,11 +44,11 @@ namespace UE::PixelStreaming2Input
 		TSharedPtr<IPixelStreaming2InputMessage> AddMessageInternal(FString Key, uint8 Id, TArray<EPixelStreaming2MessageTypes> InStructure);
 
 	private:
-		TSet<uint8>													  Ids;
+		TSet<uint8>												Ids;
 		TMap<FString, TSharedPtr<IPixelStreaming2InputMessage>> InnerMap;
-		FOnProtocolUpdated											  OnProtocolUpdatedDelegate;
-		EPixelStreaming2MessageDirection						  Direction;
-		uint8														  UserMessageId = 200;
+		FOnProtocolUpdated										OnProtocolUpdatedDelegate;
+		EPixelStreaming2MessageDirection						Direction;
+		uint8													UserMessageId = 200;
 	};
 
 } // namespace UE::PixelStreaming2Input
diff --git a/Engine/Plugins/Media/PixelStreaming2/Source/PixelStreaming2Input/Private/DefaultDataProtocol.cpp b/Engine/Plugins/Media/PixelStreaming2/Source/PixelStreaming2Input/Private/DefaultDataProtocol.cpp
--- a/Engine/Plugins/Media/PixelStreaming2/Source/PixelStreaming2Input/Private/DefaultDataProtocol.cpp
+++ b/Engine/Plugins/Media/PixelStreaming2/Source/PixelStreaming2Input/Private/DefaultDataProtocol.cpp
@@ -13,7 +13,7 @@ namespace UE::PixelStreaming2Input
 	{
 		TSharedPtr<FInputProtocolMap> ToStreamerProtocol = TSharedPtr<FInputProtocolMap>(new FInputProtocolMap(EPixelStreaming2MessageDirection::ToStreamer));
 
-		/*
+		/**
 		 * Control Messages.
 		 */
 		// Simple commands with no payload
@@ -27,7 +27,7 @@ namespace UE::PixelStreaming2Input
 		ToStreamerProtocol->AddInternal(EPixelStreaming2ToStreamerMessage::RequestInitialSettings, 7);
 		ToStreamerProtocol->AddInternal(EPixelStreaming2ToStreamerMessage::TestEcho, 8, { EType::String });
 
-		/*
+		/**
 		 * Input Messages.
 		 */
 		// Generic Input Messages.
diff --git a/Engine/Plugins/Media/PixelStreaming2/Source/PixelStreaming2Input/Private/InputHandler.cpp b/Engine/Plugins/Media/PixelStreaming2/Source/PixelStreaming2Input/Private/InputHandler.cpp
--- a/Engine/Plugins/Media/PixelStreaming2/Source/PixelStreaming2Input/Private/InputHandler.cpp
+++ b/Engine/Plugins/Media/PixelStreaming2/Source/PixelStreaming2Input/Private/InputHandler.cpp
@@ -1606,7 +1606,7 @@ namespace UE::PixelStreaming2Input
 			Buffer << Descriptor;
 			TArray<uint8> Data(Buffer.GetData(), Buffer.Num());
 
-			/*
+			/**
 			 * Send the following JSON string in a "Command" message to all players
 			 * {
 			 *	"command": "onScreenKeyboard",
diff --git a/Engine/Plugins/Media/PixelStreaming2/Source/PixelStreaming2Input/Private/InputHandler.h b/Engine/Plugins/Media/PixelStreaming2/Source/PixelStreaming2Input/Private/InputHandler.h
--- a/Engine/Plugins/Media/PixelStreaming2/Source/PixelStreaming2Input/Private/InputHandler.h
+++ b/Engine/Plugins/Media/PixelStreaming2/Source/PixelStreaming2Input/Private/InputHandler.h
@@ -24,13 +24,13 @@ namespace UE::PixelStreaming2Input
 
 		virtual void Tick(float DeltaTime) override;
 
-		/** Poll for controller state and send events if needed */
+		// Poll for controller state and send events if needed
 		virtual void SendControllerEvents() override {};
 
-		/** Set which MessageHandler will route input  */
+		// Set which MessageHandler will route input
 		virtual void SetMessageHandler(const TSharedRef<FGenericApplicationMessageHandler>& InTargetHandler) override;
 
-		/** Register a custom function to execute when command JSON is received. */
+		// Register a custom function to execute when command JSON is received.
 		virtual void SetCommandHandler(const FString& CommandName, const CommandHandlerFn& Handler) override;
 
 		virtual void SetElevatedCheck(const TFunction<bool(FString)>& CheckFn) override;
@@ -39,8 +39,9 @@ namespace UE::PixelStreaming2Input
 		virtual TSharedPtr<IPixelStreaming2DataProtocol> GetToStreamerProtocol() override;
 		virtual TSharedPtr<IPixelStreaming2DataProtocol> GetFromStreamerProtocol() override;
 
-		/** Exec handler to allow console commands to be passed through for debugging */
+		// Exec handler to allow console commands to be passed through for debugging
 		virtual bool Exec(UWorld* InWorld, const TCHAR* Cmd, FOutputDevice& Ar) override;
+
 		/**
 		 * IInputInterface pass through functions
 		 */
@@ -151,14 +152,14 @@ namespace UE::PixelStreaming2Input
 
 		struct FAnalogValue
 		{
-			/* The actual analog value from the controller axis, typical 0.0..1.0 */
+			/** The actual analog value from the controller axis, typical 0.0..1.0 */
 			double Value;
-			/* If value is non-zero then keep applying this analog values across frames.
+			/** If value is non-zero then keep applying this analog values across frames.
 			 * This is useful for trigger axis inputs where if a value is not transmitted
 			 * UE will assume a gap in input means a full trigger press (which is not accurate if we were still pressing).
 			 */
 			bool bKeepUnlessZero = false;
-			/* Has this key event already been fired once? */
+			/** Has this key event already been fired once? */
 			bool bIsRepeat = false;
 		};
 
@@ -201,16 +202,16 @@ namespace UE::PixelStreaming2Input
 			TArray<uint8>							 Data;
 		};
 
-		TWeakPtr<SWindow>				TargetWindow;
-		TWeakPtr<SViewport>				TargetViewport;
-		TWeakPtr<FIntPoint>				TargetScreenSize; // Deprecated functionality but remaining until it can be removed
-		TWeakPtr<FIntRect>				TargetScreenRect; // Manual size override used when we don't have a single window/viewport target
-		uint8							NumActiveTouches;
-		bool							bIsMouseActive;
-		TQueue<FMessage>				Messages;
-		EPixelStreaming2InputType InputType = EPixelStreaming2InputType::RouteToWindow;
-		FVector2D						LastTouchLocation = FVector2D(EForceInit::ForceInitToZero);
-		TMap<uint8, MessageHandlerFn>	DispatchTable;
+		TWeakPtr<SWindow>			  TargetWindow;
+		TWeakPtr<SViewport>			  TargetViewport;
+		TWeakPtr<FIntPoint>			  TargetScreenSize; // Deprecated functionality but remaining until it can be removed
+		TWeakPtr<FIntRect>			  TargetScreenRect; // Manual size override used when we don't have a single window/viewport target
+		uint8						  NumActiveTouches;
+		bool						  bIsMouseActive;
+		TQueue<FMessage>			  Messages;
+		EPixelStreaming2InputType	  InputType = EPixelStreaming2InputType::RouteToWindow;
+		FVector2D					  LastTouchLocation = FVector2D(EForceInit::ForceInitToZero);
+		TMap<uint8, MessageHandlerFn> DispatchTable;
 
 		/** Reference to the message handler which events should be passed to. */
 		TSharedPtr<FGenericApplicationMessageHandler> MessageHandler;
@@ -225,7 +226,7 @@ namespace UE::PixelStreaming2Input
 		 */
 		bool bFakingTouchEvents;
 
-		/*
+		/**
 		 * Padding for string parsing when handling messages.
 		 * 1 character for the actual message and then
 		 * 2 characters for the length which are skipped
diff --git a/Engine/Plugins/Media/PixelStreaming2/Source/PixelStreaming2Input/Private/Utils.h b/Engine/Plugins/Media/PixelStreaming2/Source/PixelStreaming2Input/Private/Utils.h
--- a/Engine/Plugins/Media/PixelStreaming2/Source/PixelStreaming2Input/Private/Utils.h
+++ b/Engine/Plugins/Media/PixelStreaming2/Source/PixelStreaming2Input/Private/Utils.h
@@ -33,7 +33,7 @@ namespace UE::PixelStreaming2Input
 
 	inline void ExtractJsonFromDescriptor(FString Descriptor, FString FieldName, FString& StringValue, bool& Success)
 	{
-		/*
+		/**
 		 * ExtractJsonFromDescriptor supports parsing nested objects, as well as keys that contain a period.
 		 * eg
 		 * { "Encoder.MinQP": val }
diff --git a/Engine/Plugins/Media/PixelStreaming2/Source/PixelStreaming2Servers/Public/PixelStreaming2Servers.h b/Engine/Plugins/Media/PixelStreaming2/Source/PixelStreaming2Servers/Public/PixelStreaming2Servers.h
--- a/Engine/Plugins/Media/PixelStreaming2/Source/PixelStreaming2Servers/Public/PixelStreaming2Servers.h
+++ b/Engine/Plugins/Media/PixelStreaming2/Source/PixelStreaming2Servers/Public/PixelStreaming2Servers.h
@@ -11,7 +11,7 @@
 
 /*
  * Utility namespace for launching Pixel Streaming servers signaling or selective forwarding.
- * This utility can launch embedded servers from C++ classes in this module, 
+ * This utility can launch embedded servers from C++ classes in this module,
  * from known server binaries and from the Github releases of the servers.
  */
 namespace UE::PixelStreaming2Servers
@@ -65,7 +65,9 @@ namespace UE::PixelStreaming2Servers
 	public:
 		virtual ~IServer() = default;
 
-		/* Immediately stops the server. */
+		/**
+		 * Immediately stops the server.
+		 */
 		virtual void Stop() = 0;
 
 		/**
@@ -98,7 +100,7 @@ namespace UE::PixelStreaming2Servers
 		/**
 		 * Asynchronously queries the numbers of streamers connected to the signalling server.
 		 * @param OnNumStreamersReceived	Callback fired when the query for the number of streamers returns.
-		*/
+		 */
 		virtual void GetNumStreamers(TFunction<void(uint16)> OnNumStreamersReceived) = 0;
 
 	public:
