[PS2] Fix screen tearing in stream when using Pixel Streaming 2 with AMD gpu.

Problem occurs because the buffer is passed to AMF before the copy completes.

This change re-adds the CaptureUseFence CVar which is set to true of a AMD gpu is used. This then forces a fence to be used for the copy and fixes the screen tearing.

diff --git a/Engine/Plugins/Media/PixelStreaming2/Source/PixelStreaming2/Internal/VideoCapturer.h b/Engine/Plugins/Media/PixelStreaming2/Source/PixelStreaming2/Internal/VideoCapturer.h
--- a/Engine/Plugins/Media/PixelStreaming2/Source/PixelStreaming2/Internal/VideoCapturer.h
+++ b/Engine/Plugins/Media/PixelStreaming2/Source/PixelStreaming2/Internal/VideoCapturer.h
@@ -54,10 +54,12 @@ namespace UE::PixelStreaming2
 		TSharedPtr<FPixelCaptureCapturerMultiFormat> FrameCapturer;
 		FDelegateHandle								 CaptureCompleteHandle;
 		FDelegateHandle								 SimulcastEnabledChangedHandle;
+		FDelegateHandle								 CaptureUseFenceChangedHandle;
 		FDelegateHandle								 FramePushedHandle;
 
 		void CreateFrameCapturer();
 		void OnSimulcastEnabledChanged(IConsoleVariable* Var);
+		void OnCaptureUseFenceChanged(IConsoleVariable* Var);
 		void OnCaptureComplete();
 		void OnFrame(const IPixelCaptureInputFrame& InputFrame);
 	};
diff --git a/Engine/Plugins/Media/PixelStreaming2/Source/PixelStreaming2/Private/PixelStreaming2Module.cpp b/Engine/Plugins/Media/PixelStreaming2/Source/PixelStreaming2/Private/PixelStreaming2Module.cpp
--- a/Engine/Plugins/Media/PixelStreaming2/Source/PixelStreaming2/Private/PixelStreaming2Module.cpp
+++ b/Engine/Plugins/Media/PixelStreaming2/Source/PixelStreaming2/Private/PixelStreaming2Module.cpp
@@ -161,6 +161,19 @@ namespace UE::PixelStreaming2
 			// Ensure we have ImageWrapper loaded, used in Freezeframes
 			verify(FModuleManager::Get().LoadModule(FName("ImageWrapper")));
 
+			// HACK (Eden.Harris): Until or if we ever find a workaround for fencing, we need to ensure capture always uses a fence.
+			// If we don't then we get frequent and intermittent stuttering as textures are rendered to while being encoded.
+			// From testing NVENC + CUDA pathway seems acceptable without a fence in most cases so we use the faster, unsafer path there.
+			if (IsRHIDeviceAMD())
+			{
+				if (!UPixelStreaming2PluginSettings::CVarCaptureUseFence.GetValueOnAnyThread())
+				{
+					UE_LOGFMT(LogPixelStreaming2, Warning, "AMD GPU Device detected, setting PixelStreaming2.CaptureUseFence to true to avoid screen tearing in stream.");
+				}
+
+				UPixelStreaming2PluginSettings::CVarCaptureUseFence.AsVariable()->Set(true);
+			}
+
 			// We don't want to start immediately streaming in editor
 			if (!GIsEditor)
 			{
diff --git a/Engine/Plugins/Media/PixelStreaming2/Source/PixelStreaming2/Private/VideoCapturer.cpp b/Engine/Plugins/Media/PixelStreaming2/Source/PixelStreaming2/Private/VideoCapturer.cpp
--- a/Engine/Plugins/Media/PixelStreaming2/Source/PixelStreaming2/Private/VideoCapturer.cpp
+++ b/Engine/Plugins/Media/PixelStreaming2/Source/PixelStreaming2/Private/VideoCapturer.cpp
@@ -8,6 +8,7 @@
 #include "PixelCaptureCapturerI420.h"
 #include "PixelCaptureCapturerI420ToRHI.h"
 #include "PixelCaptureCapturerNV12ToRHI.h"
+#include "PixelCaptureCapturerRHI.h"
 #include "PixelCaptureCapturerRHIRDG.h"
 #include "PixelCaptureCapturerRHIToI420CPU.h"
 #include "PixelStreaming2PluginSettings.h"
@@ -29,6 +30,7 @@ namespace UE::PixelStreaming2
 		if (UPixelStreaming2PluginSettings::FDelegates* Delegates = UPixelStreaming2PluginSettings::Delegates())
 		{
 			VideoCapturer->SimulcastEnabledChangedHandle = Delegates->OnSimulcastEnabledChanged.AddSP(VideoCapturer.ToSharedRef(), &FVideoCapturer::OnSimulcastEnabledChanged);
+			VideoCapturer->CaptureUseFenceChangedHandle = Delegates->OnCaptureUseFenceChanged.AddSP(VideoCapturer.ToSharedRef(), &FVideoCapturer::OnCaptureUseFenceChanged);
 		}
 		return VideoCapturer;
 	}
@@ -106,6 +108,11 @@ namespace UE::PixelStreaming2
 		CreateFrameCapturer();
 	}
 
+	void FVideoCapturer::OnCaptureUseFenceChanged(IConsoleVariable* Var)
+	{
+		CreateFrameCapturer();
+	}
+
 	void FVideoCapturer::CreateFrameCapturer()
 	{
 		if (FrameCapturer != nullptr)
@@ -138,7 +145,14 @@ namespace UE::PixelStreaming2
 			switch (FinalFormat)
 			{
 				case PixelCaptureBufferFormat::FORMAT_RHI:
-					return FPixelCaptureCapturerRHIRDG::Create(FinalScale);
+					if(UPixelStreaming2PluginSettings::CVarCaptureUseFence.GetValueOnAnyThread())
+					{
+						return FPixelCaptureCapturerRHI::Create(FinalScale);
+					}
+					else 
+					{
+						return FPixelCaptureCapturerRHIRDG::Create(FinalScale);
+					}
 				case PixelCaptureBufferFormat::FORMAT_I420:
 					return FPixelCaptureCapturerRHIToI420CPU::Create(FinalScale);
 				default:
diff --git a/Engine/Plugins/Media/PixelStreaming2/Source/PixelStreaming2Core/Internal/PixelStreaming2PluginSettings.h b/Engine/Plugins/Media/PixelStreaming2/Source/PixelStreaming2Core/Internal/PixelStreaming2PluginSettings.h
--- a/Engine/Plugins/Media/PixelStreaming2/Source/PixelStreaming2Core/Internal/PixelStreaming2PluginSettings.h
+++ b/Engine/Plugins/Media/PixelStreaming2/Source/PixelStreaming2Core/Internal/PixelStreaming2PluginSettings.h
@@ -153,6 +153,14 @@ public:
 		))
 	FString SignallingURL;
 
+	static TAutoConsoleVariable<bool> CVarCaptureUseFence;
+	UPROPERTY(config, EditAnywhere, Category = "PixelStreaming", meta = (
+		ConsoleVariable = "PixelStreaming2.CaptureUseFence",
+		DisplayName = "Capture Using Fence",
+		ToolTip = "Whether the texture copy we do during image capture should use a fence or not (non-fenced is faster but less safe)."
+		))
+	bool CaptureUseFence = false;
+
 	static TAutoConsoleVariable<bool> CVarDebugDumpAudio;
 
 	// Begin Cursor Settings
@@ -612,6 +620,9 @@ public:
 		DECLARE_TS_MULTICAST_DELEGATE_OneParam(FOnSimulcastEnabledChanged, IConsoleVariable*);
 		FOnSimulcastEnabledChanged OnSimulcastEnabledChanged;
 
+		DECLARE_TS_MULTICAST_DELEGATE_OneParam(FOnCaptureUseFenceChanged, IConsoleVariable*);
+		FOnCaptureUseFenceChanged OnCaptureUseFenceChanged;
+
 		DECLARE_TS_MULTICAST_DELEGATE_OneParam(FOnWebRTCFpsChanged, IConsoleVariable*);
 		FOnWebRTCFpsChanged OnWebRTCFpsChanged;
 
diff --git a/Engine/Plugins/Media/PixelStreaming2/Source/PixelStreaming2Core/Private/PixelStreaming2PluginSettings.cpp b/Engine/Plugins/Media/PixelStreaming2/Source/PixelStreaming2Core/Private/PixelStreaming2PluginSettings.cpp
--- a/Engine/Plugins/Media/PixelStreaming2/Source/PixelStreaming2Core/Private/PixelStreaming2PluginSettings.cpp
+++ b/Engine/Plugins/Media/PixelStreaming2/Source/PixelStreaming2Core/Private/PixelStreaming2PluginSettings.cpp
@@ -202,6 +202,13 @@ TAutoConsoleVariable<FString> UPixelStreaming2PluginSettings::CVarSignallingURL(
 	TEXT("Default URL to connect to for signalling."),
 	ECVF_Default);
 
+TAutoConsoleVariable<bool> UPixelStreaming2PluginSettings::CVarCaptureUseFence(
+	TEXT("PixelStreaming2.CaptureUseFence"),
+	false,
+	TEXT("Whether the texture copy we do during image capture should use a fence or not (non-fenced is faster but less safe)."),
+	FConsoleVariableDelegate::CreateLambda([](IConsoleVariable* Var) { Delegates()->OnCaptureUseFenceChanged.Broadcast(Var); }),
+	ECVF_Default);
+
 TAutoConsoleVariable<bool> UPixelStreaming2PluginSettings::CVarDebugDumpAudio(
 	TEXT("PixelStreaming2.DumpDebugAudio"),
 	false,
@@ -911,6 +918,7 @@ static const TMap<FString, FString> GetCmdArg = {
 	{ "UseMediaCapture", "PixelStreaming2.UseMediaCapture" },
 	{ "DefaultStreamerID", "PixelStreaming2.ID" },
 	{ "SignallingURL", "PixelStreaming2.SignallingURL" },
+	{ "CaptureUseFence", "PixelStreaming2.CaptureUseFence" },
 	{ "Codec", "PixelStreaming2.Encoder.Codec" },
 	{ "EncoderTargetBitrate", "PixelStreaming2.Encoder.TargetBitrate" },
 	{ "EncoderMinQuality", "PixelStreaming2.Encoder.MinQuality" },
