[PS2] Fix: Decoder resource pooling

diff --git a/Engine/Plugins/Experimental/AVCodecs/AVCodecsCore/Source/AVCodecsCore/Public/Video/VideoResource.h b/Engine/Plugins/Experimental/AVCodecs/AVCodecsCore/Source/AVCodecsCore/Public/Video/VideoResource.h
--- a/Engine/Plugins/Experimental/AVCodecs/AVCodecsCore/Source/AVCodecsCore/Public/Video/VideoResource.h
+++ b/Engine/Plugins/Experimental/AVCodecs/AVCodecsCore/Source/AVCodecsCore/Public/Video/VideoResource.h
@@ -8,6 +8,8 @@
 #include "PixelFormat.h"
 #include "Containers/ResourceArray.h"
 
+#include <atomic>
+
 enum class EVideoFormat : uint8
 {
 	BGRA = EPixelFormat::PF_B8G8R8A8,
@@ -167,6 +169,11 @@ private:
 	 */
 	FVideoDescriptor Descriptor;
 
+	/**
+	 * 
+	 */
+	std::atomic<bool> bIsInUse = false;
+
 public:
 	/**
 	 * @return Get the descriptor of our video data in device memory.
@@ -198,6 +205,16 @@ public:
 	 */
 	FORCEINLINE uint32 GetSize() const { return Descriptor.GetSizeInBytes(); }
 
+	/**
+	 * 
+	 */
+	FORCEINLINE void SetUsing(bool bUsing) { bIsInUse = bUsing; }
+
+	/**
+	 * 
+	 */
+	FORCEINLINE bool IsInUse() { return bIsInUse; }
+
 	FVideoResource(TSharedRef<FAVDevice> const& Device, FAVLayout const& Layout, FVideoDescriptor const& Descriptor);
 	virtual ~FVideoResource() override = default;
 };
diff --git a/Engine/Plugins/Media/PixelStreaming2/Source/PixelStreaming2/Private/EpicRtcVideoBufferRHI.h b/Engine/Plugins/Media/PixelStreaming2/Source/PixelStreaming2/Private/EpicRtcVideoBufferRHI.h
--- a/Engine/Plugins/Media/PixelStreaming2/Source/PixelStreaming2/Private/EpicRtcVideoBufferRHI.h
+++ b/Engine/Plugins/Media/PixelStreaming2/Source/PixelStreaming2/Private/EpicRtcVideoBufferRHI.h
@@ -14,9 +14,13 @@ namespace UE::PixelStreaming2
 		FEpicRtcVideoBufferRHI(TSharedPtr<FVideoResourceRHI> VideoResourceRHI)
 			: VideoResourceRHI(VideoResourceRHI) 
 		{
+			VideoResourceRHI->SetUsing(true);
 		}
 
-		virtual ~FEpicRtcVideoBufferRHI() = default;
+		virtual ~FEpicRtcVideoBufferRHI() 
+		{
+			VideoResourceRHI->SetUsing(false);
+		}
 
 	public:
 		// Begin EpicRtcVideoBufferInterface
diff --git a/Engine/Plugins/Media/PixelStreaming2/Source/PixelStreaming2/Private/EpicRtcVideoDecoder.cpp b/Engine/Plugins/Media/PixelStreaming2/Source/PixelStreaming2/Private/EpicRtcVideoDecoder.cpp
--- a/Engine/Plugins/Media/PixelStreaming2/Source/PixelStreaming2/Private/EpicRtcVideoDecoder.cpp
+++ b/Engine/Plugins/Media/PixelStreaming2/Source/PixelStreaming2/Private/EpicRtcVideoDecoder.cpp
@@ -116,8 +116,7 @@ namespace UE::PixelStreaming2
 			FAVResult DecodeResult;
 			if constexpr (std::is_same_v<TVideoResource, FVideoResourceRHI>)
 			{
-				// HACK edge case if streaming framerate is faster than render framerate this will cause an issue and should be switched to FetchOrCreate
-				FResolvableVideoResourceRHI& DecoderResource = VideoResourcesRHI[VideoResourceIndex++ % 2];
+				FResolvableVideoResourceRHI& DecoderResource = VideoResources.GetOrCreate();
 				DecodeResult = Decoder->ReceiveFrame(DecoderResource);
 
 				if (DecodeResult.IsSuccess())
@@ -127,8 +126,7 @@ namespace UE::PixelStreaming2
 			}
 			else if constexpr (std::is_same_v<TVideoResource, FVideoResourceCPU>)
 			{
-				// HACK edge case if streaming framerate is faster than render framerate this will cause an issue and should be switched to FetchOrCreate
-				FResolvableVideoResourceCPU& DecoderResource = VideoResourcesCPU[VideoResourceIndex++ % 2];
+				FResolvableVideoResourceCPU& DecoderResource = VideoResources.GetOrCreate();
 				DecodeResult = Decoder->ReceiveFrame(DecoderResource);
 
 				if (DecodeResult.IsSuccess())
diff --git a/Engine/Plugins/Media/PixelStreaming2/Source/PixelStreaming2/Private/EpicRtcVideoDecoder.h b/Engine/Plugins/Media/PixelStreaming2/Source/PixelStreaming2/Private/EpicRtcVideoDecoder.h
--- a/Engine/Plugins/Media/PixelStreaming2/Source/PixelStreaming2/Private/EpicRtcVideoDecoder.h
+++ b/Engine/Plugins/Media/PixelStreaming2/Source/PixelStreaming2/Private/EpicRtcVideoDecoder.h
@@ -2,6 +2,7 @@
 
 #pragma once
 
+#include "IPixelStreaming2Stats.h"
 #include "Video/Decoders/Configs/VideoDecoderConfigAV1.h"
 #include "Video/Decoders/Configs/VideoDecoderConfigH264.h"
 #include "Video/Decoders/Configs/VideoDecoderConfigVP8.h"
@@ -13,9 +14,53 @@
 #include "epic_rtc/core/video/video_decoder.h"
 
 #include <atomic>
+#include <type_traits>
 
 namespace UE::PixelStreaming2
 {
+	template <typename TResolvableVideoResource, typename TVideoResource>
+	class TVideoResourcePool
+	{
+	public:
+		TResolvableVideoResource& GetOrCreate()
+		{
+			ON_SCOPE_EXIT
+			{
+				IPixelStreaming2Stats::Get().GraphValue(TEXT("NumDecodeResource"), Resources.Num(), 1, 0.f, 120.f);
+			};
+
+			if (RHIGetInterfaceType() == ERHIInterfaceType::Vulkan)
+			{
+				// Vulkan has lifetime issues so we always use the same resource
+				if (Resources.Num() == 0)
+				{
+					return Resources.Emplace_GetRef();
+				}
+				return Resources[0];
+			}
+
+			TResolvableVideoResource* Resource = Resources.FindByPredicate([](TResolvableVideoResource& ResolvableResource) {
+				TSharedPtr<TVideoResource>& Resolved = ResolvableResource; // Activates the TSharedPtr operator overload which will resolve the resource
+				if (!Resolved.IsValid())
+				{
+					return false;
+				}
+
+				return !Resolved->IsInUse();
+			});
+
+			if (!Resource)
+			{
+				return Resources.Emplace_GetRef();
+			}
+
+			return *Resource;
+		}
+
+	private:
+		TArray<TResolvableVideoResource> Resources;
+	};
+
 	template <std::derived_from<FVideoResource> TVideoResource>
 	class TEpicRtcVideoDecoder : public EpicRtcVideoDecoderInterface, public TRefCountingMixin<TEpicRtcVideoDecoder<TVideoResource>>
 	{
@@ -37,10 +82,16 @@ namespace UE::PixelStreaming2
 		TRefCountPtr<EpicRtcVideoDecoderCallbackInterface> VideoDecoderCallback;
 		TRefCountPtr<EpicRtcVideoCodecInfoInterface>	   CodecInfo;
 		uint16_t										   FrameCount;
-		// TODO rather than just flip flopping we should have a system to FetchOrCreate a Resource on an as needed basis
-		TStaticArray<FResolvableVideoResourceRHI, 2>	   VideoResourcesRHI;
-		TStaticArray<FResolvableVideoResourceCPU, 2>	   VideoResourcesCPU;
-		std::atomic<uint32>								   VideoResourceIndex;
+
+		// clang-format off
+		using ResourcePoolType = typename std::conditional<std::is_same_v<TVideoResource, FVideoResourceRHI>, TVideoResourcePool<FResolvableVideoResourceRHI, FVideoResourceRHI>, 
+								 typename std::conditional<std::is_same_v<TVideoResource, FVideoResourceCPU>, TVideoResourcePool<FResolvableVideoResourceCPU, FVideoResourceCPU>, 
+								 // We only support RHI and CPU resources so if we create one that isn't this type it will error (on purpose)
+								 void>::type>::type; 
+
+		ResourcePoolType VideoResources;
+		// clang-format on
+
 	private:
 		bool LateInitDecoder();
 
