Fix non-unity build for MacCriticalSection.cpp and PixelStreaming2MediaTexture.cpp

diff --git a/Engine/Plugins/Media/PixelStreaming2/Source/PixelStreaming2/Private/Blueprints/PixelStreaming2MediaTexture.cpp b/Engine/Plugins/Media/PixelStreaming2/Source/PixelStreaming2/Private/Blueprints/PixelStreaming2MediaTexture.cpp
--- a/Engine/Plugins/Media/PixelStreaming2/Source/PixelStreaming2/Private/Blueprints/PixelStreaming2MediaTexture.cpp
+++ b/Engine/Plugins/Media/PixelStreaming2/Source/PixelStreaming2/Private/Blueprints/PixelStreaming2MediaTexture.cpp
@@ -5,6 +5,7 @@
 #include "Async/Async.h"
 #include "DeviceProfiles/DeviceProfile.h"
 #include "DeviceProfiles/DeviceProfileManager.h"
+#include "RenderUtils.h"
 #include "TextureResource.h"
 
 constexpr int32 DEFAULT_WIDTH = 1920;
diff --git a/Engine/Source/Runtime/Core/Private/Mac/MacCriticalSection.cpp b/Engine/Source/Runtime/Core/Private/Mac/MacCriticalSection.cpp
--- a/Engine/Source/Runtime/Core/Private/Mac/MacCriticalSection.cpp
+++ b/Engine/Source/Runtime/Core/Private/Mac/MacCriticalSection.cpp
@@ -5,6 +5,7 @@
 =============================================================================*/
 
 #include "Mac/MacCriticalSection.h"
+#include "HAL/PlatformTime.h"
 #include "Misc/App.h"
 #include "Misc/DateTime.h"
 #include "Apple/PreAppleSystemHeaders.h"
