[PS2] Fix Pixel Streaming 2 protocol tests crash.

The protocol tests did not set the codec to vp8 and would crash if the codec was wet to h264 on a machine that did not support decoding h264. It would cause webrtc to fail to create a valid sdp response.

This fix sets the codec to vp8 which is software encoded.decoded.

diff --git a/Engine/Plugins/Media/PixelStreaming2/Source/PixelStreaming2/Private/Tests/ProtocolTest.cpp b/Engine/Plugins/Media/PixelStreaming2/Source/PixelStreaming2/Private/Tests/ProtocolTest.cpp
--- a/Engine/Plugins/Media/PixelStreaming2/Source/PixelStreaming2/Private/Tests/ProtocolTest.cpp
+++ b/Engine/Plugins/Media/PixelStreaming2/Source/PixelStreaming2/Private/Tests/ProtocolTest.cpp
@@ -18,13 +18,16 @@ namespace UE::PixelStreaming2
 	IMPLEMENT_SIMPLE_AUTOMATION_TEST(FProtocolTestAddMessage, "System.Plugins.PixelStreaming2.FPS2ProtocolTestAddMessage", EAutomationTestFlags::EditorContext | EAutomationTestFlags::ClientContext | EAutomationTestFlags::ProductFilter)
 	bool FProtocolTestAddMessage::RunTest(const FString& Parameters)
 	{
+		// need to be able to accept codec to handshake otherwise setting local description fails when generating an answer
+		SetCodec(EVideoCodec::VP8);
+		
 		int32 StreamerPort = TestUtils::NextStreamerPort();
 		int32 PlayerPort = TestUtils::NextPlayerPort();
 
 		TSharedPtr<UE::PixelStreaming2Servers::IServer> SignallingServer = CreateSignallingServer(StreamerPort, PlayerPort);
 
-		FString										   StreamerName(FString::Printf(TEXT("MockStreamer%d"), StreamerPort));
-		TSharedPtr<IPixelStreaming2Streamer>	   Streamer = CreateStreamer(StreamerName, StreamerPort);
+		FString									 StreamerName(FString::Printf(TEXT("MockStreamer%d"), StreamerPort));
+		TSharedPtr<IPixelStreaming2Streamer>	 Streamer = CreateStreamer(StreamerName, StreamerPort);
 		TSharedPtr<IPixelStreaming2InputHandler> InputHandler = Streamer->GetInputHandler().Pin();
 
 		const FString CustomMessageName = FString(TEXT("CustomMessage"));
@@ -93,6 +96,9 @@ namespace UE::PixelStreaming2
 	IMPLEMENT_SIMPLE_AUTOMATION_TEST(FProtocolTestUseCustomMessage, "System.Plugins.PixelStreaming2.FPS2ProtocolTestUseCustomMessage", EAutomationTestFlags::EditorContext | EAutomationTestFlags::ClientContext | EAutomationTestFlags::ProductFilter)
 	bool FProtocolTestUseCustomMessage::RunTest(const FString& Parameters)
 	{
+		// need to be able to accept codec to handshake otherwise setting local description fails when generating an answer
+		SetCodec(EVideoCodec::VP8);
+		
 		int32 StreamerPort = TestUtils::NextStreamerPort();
 		int32 PlayerPort = TestUtils::NextPlayerPort();
 
