[PS2]: Fix clashing test names in macro definitions and stuttering audio

diff --git a/Engine/Plugins/Media/PixelStreaming2/Source/PixelStreaming2/Private/EpicRtcThread.cpp b/Engine/Plugins/Media/PixelStreaming2/Source/PixelStreaming2/Private/EpicRtcThread.cpp
--- a/Engine/Plugins/Media/PixelStreaming2/Source/PixelStreaming2/Private/EpicRtcThread.cpp
+++ b/Engine/Plugins/Media/PixelStreaming2/Source/PixelStreaming2/Private/EpicRtcThread.cpp
@@ -34,7 +34,8 @@ namespace UE::PixelStreaming2
 			{
 				Tick();
 
-				FPlatformProcess::Sleep(0.f);
+				// Sleep 1ms
+				FPlatformProcess::Sleep(0.001f);
 			}
 
 			return 0;
diff --git a/Engine/Plugins/Media/PixelStreaming2/Source/PixelStreaming2/Private/EpicRtcVideoEncoder.cpp b/Engine/Plugins/Media/PixelStreaming2/Source/PixelStreaming2/Private/EpicRtcVideoEncoder.cpp
--- a/Engine/Plugins/Media/PixelStreaming2/Source/PixelStreaming2/Private/EpicRtcVideoEncoder.cpp
+++ b/Engine/Plugins/Media/PixelStreaming2/Source/PixelStreaming2/Private/EpicRtcVideoEncoder.cpp
@@ -11,6 +11,7 @@
 #include "PixelCaptureOutputFrameI420.h"
 #include "PixelCaptureOutputFrameRHI.h"
 #include "PixelStreaming2PluginSettings.h"
+#include "Stats.h"
 #include "UtilsString.h"
 #include "UtilsVideo.h"
 
@@ -377,8 +378,9 @@ namespace UE::PixelStreaming2
 		// Update the encoding config using the incoming frame resolution (required for dynamic res support)
 		UpdateConfig(Width, Height);
 
-		TSharedPtr<TVideoResource> VideoResource;
+		UpdateFrameMetadataPreEncode(*AdaptedLayer);
 
+		TSharedPtr<TVideoResource> VideoResource;
 		if constexpr (std::is_same_v<TVideoResource, FVideoResourceRHI>)
 		{
 			const FPixelCaptureOutputFrameRHI& RHILayer = StaticCast<const FPixelCaptureOutputFrameRHI&>(*AdaptedLayer);
@@ -418,6 +420,8 @@ namespace UE::PixelStreaming2
 		// Encode
 		Encoder->SendFrame(VideoResource, VideoFrame._timestampUs, bKeyFrame);
 
+		UpdateFrameMetadataPostEncode(*AdaptedLayer);
+
 		FVideoPacket Packet;
 		while (Encoder->ReceivePacket(Packet))
 		{
@@ -524,7 +528,10 @@ namespace UE::PixelStreaming2
 
 			MaybeDumpFrame(EncodedFrame);
 
+			UpdateFrameMetadataPrePacketization(*AdaptedLayer);
 			EpicRtcVideoEncodedResult Result = VideoEncoderCallback->Encoded(EncodedFrame, CodecSpecificInfo);
+			UpdateFrameMetadataPostPacketization(*AdaptedLayer);
+
 			if (Result._error)
 			{
 				UE_LOGFMT(LogPixelStreaming2, Error, "PixelStreamingVideoEncoder: Failed to push encoded frame");
@@ -738,6 +745,47 @@ namespace UE::PixelStreaming2
 		return SumBps;
 	}
 
+	template <std::derived_from<FVideoResource> TVideoResource>
+	void TEpicRtcVideoEncoder<TVideoResource>::UpdateFrameMetadataPreEncode(IPixelCaptureOutputFrame& Frame)
+	{
+		FPixelCaptureFrameMetadata& FrameMetadata = Frame.Metadata;
+		FrameMetadata.UseCount++;
+		FrameMetadata.LastEncodeStartTime = FPlatformTime::ToMilliseconds64(FPlatformTime::Cycles64());
+		if (FrameMetadata.UseCount == 1)
+		{
+			FrameMetadata.FirstEncodeStartTime = FrameMetadata.LastEncodeStartTime;
+		}
+	}
+
+	template <std::derived_from<FVideoResource> TVideoResource>
+	void TEpicRtcVideoEncoder<TVideoResource>::UpdateFrameMetadataPostEncode(IPixelCaptureOutputFrame& Frame)
+	{
+		FPixelCaptureFrameMetadata& FrameMetadata = Frame.Metadata;
+		FrameMetadata.LastEncodeEndTime = FPlatformTime::ToMilliseconds64(FPlatformTime::Cycles64());
+
+		FStats::Get()->AddFrameTimingStats(FrameMetadata);
+	}
+
+	template <std::derived_from<FVideoResource> TVideoResource>
+	void TEpicRtcVideoEncoder<TVideoResource>::UpdateFrameMetadataPrePacketization(IPixelCaptureOutputFrame& Frame)
+	{
+		FPixelCaptureFrameMetadata& FrameMetadata = Frame.Metadata;
+		FrameMetadata.LastPacketizationStartTime = FPlatformTime::ToMilliseconds64(FPlatformTime::Cycles64());
+		if (FrameMetadata.UseCount == 1)
+		{
+			FrameMetadata.FirstPacketizationStartTime = FrameMetadata.LastPacketizationStartTime;
+		}
+	}
+
+	template <std::derived_from<FVideoResource> TVideoResource>
+	void TEpicRtcVideoEncoder<TVideoResource>::UpdateFrameMetadataPostPacketization(IPixelCaptureOutputFrame& Frame)
+	{
+		FPixelCaptureFrameMetadata& FrameMetadata = Frame.Metadata;
+		FrameMetadata.LastPacketizationEndTime = FPlatformTime::ToMilliseconds64(FPlatformTime::Cycles64());
+
+		FStats::Get()->AddFrameTimingStats(FrameMetadata);
+	}
+
 	// Explicit specialisation
 	template class TEpicRtcVideoEncoder<FVideoResourceRHI>;
 	template class TEpicRtcVideoEncoder<FVideoResourceCPU>;
diff --git a/Engine/Plugins/Media/PixelStreaming2/Source/PixelStreaming2/Private/EpicRtcVideoEncoder.h b/Engine/Plugins/Media/PixelStreaming2/Source/PixelStreaming2/Private/EpicRtcVideoEncoder.h
--- a/Engine/Plugins/Media/PixelStreaming2/Source/PixelStreaming2/Private/EpicRtcVideoEncoder.h
+++ b/Engine/Plugins/Media/PixelStreaming2/Source/PixelStreaming2/Private/EpicRtcVideoEncoder.h
@@ -5,6 +5,7 @@
 #include "AVConstants.h"
 #include "EpicRtcVideoCommon.h"
 #include "HAL/PlatformFileManager.h"
+#include "IPixelCaptureOutputFrame.h"
 #include "Video/Encoders/Configs/VideoEncoderConfigAV1.h"
 #include "Video/Encoders/Configs/VideoEncoderConfigH264.h"
 #include "Video/Encoders/Configs/VideoEncoderConfigVP8.h"
@@ -59,8 +60,14 @@ namespace UE::PixelStreaming2
 		void	 MaybeDumpFrame(EpicRtcEncodedVideoFrame const& EncodedImage);
 		uint32_t SumAndResetBitrates();
 
+		void UpdateFrameMetadataPreEncode(IPixelCaptureOutputFrame& Frame);
+		void UpdateFrameMetadataPostEncode(IPixelCaptureOutputFrame& Frame);
+		void UpdateFrameMetadataPrePacketization(IPixelCaptureOutputFrame& Frame);
+		void UpdateFrameMetadataPostPacketization(IPixelCaptureOutputFrame& Frame);
+
 		void OnEncoderDebugDumpFrameChanged(IConsoleVariable* Var);
 		void CreateDumpFile();
+
 		FDelegateHandle DelegateHandle;
 
 	public:
diff --git a/Engine/Plugins/Media/PixelStreaming2/Source/PixelStreaming2/Private/Tests/DataChannelTest.cpp b/Engine/Plugins/Media/PixelStreaming2/Source/PixelStreaming2/Private/Tests/DataChannelTest.cpp
--- a/Engine/Plugins/Media/PixelStreaming2/Source/PixelStreaming2/Private/Tests/DataChannelTest.cpp
+++ b/Engine/Plugins/Media/PixelStreaming2/Source/PixelStreaming2/Private/Tests/DataChannelTest.cpp
@@ -12,8 +12,8 @@
 
 namespace UE::PixelStreaming2
 {
-	IMPLEMENT_SIMPLE_AUTOMATION_TEST(FDataChannelEchoTest, "System.Plugins.PixelStreaming2.FPS2DataChannelEchoTest", EAutomationTestFlags::EditorContext | EAutomationTestFlags::ClientContext | EAutomationTestFlags::ProductFilter)
-	bool FDataChannelEchoTest::RunTest(const FString& Parameters)
+	IMPLEMENT_SIMPLE_AUTOMATION_TEST(FPS2DataChannelEchoTest, "System.Plugins.PixelStreaming2.FPS2DataChannelEchoTest", EAutomationTestFlags::EditorContext | EAutomationTestFlags::ClientContext | EAutomationTestFlags::ProductFilter)
+	bool FPS2DataChannelEchoTest::RunTest(const FString& Parameters)
 	{
 		// need to be able to accept codec to handshake otherwise setting local description fails when generating an answer
 		SetCodec(EVideoCodec::VP8);
diff --git a/Engine/Plugins/Media/PixelStreaming2/Source/PixelStreaming2/Private/Tests/EpicRtcConferenceTest.cpp b/Engine/Plugins/Media/PixelStreaming2/Source/PixelStreaming2/Private/Tests/EpicRtcConferenceTest.cpp
--- a/Engine/Plugins/Media/PixelStreaming2/Source/PixelStreaming2/Private/Tests/EpicRtcConferenceTest.cpp
+++ b/Engine/Plugins/Media/PixelStreaming2/Source/PixelStreaming2/Private/Tests/EpicRtcConferenceTest.cpp
@@ -13,8 +13,8 @@
 namespace UE::PixelStreaming2
 {
 	// Tests the creation and removal of a conference from the platform
-	IMPLEMENT_SIMPLE_AUTOMATION_TEST(FEpicRtcConferenceLifetimeTest, "System.Plugins.PixelStreaming2.FPS2EpicRtcConferenceLifetimeTest", EAutomationTestFlags::EditorContext | EAutomationTestFlags::ClientContext | EAutomationTestFlags::ProductFilter)
-	bool FEpicRtcConferenceLifetimeTest::RunTest(const FString& Parameters)
+	IMPLEMENT_SIMPLE_AUTOMATION_TEST(FPS2EpicRtcConferenceLifetimeTest, "System.Plugins.PixelStreaming2.FPS2EpicRtcConferenceLifetimeTest", EAutomationTestFlags::EditorContext | EAutomationTestFlags::ClientContext | EAutomationTestFlags::ProductFilter)
+	bool FPS2EpicRtcConferenceLifetimeTest::RunTest(const FString& Parameters)
 	{
 		EpicRtcErrorCode Result;
 
@@ -83,8 +83,8 @@ namespace UE::PixelStreaming2
 	}
 
 	// Tests the conference creation logic to ensure that two conferences with the same name can't be created
-	IMPLEMENT_SIMPLE_AUTOMATION_TEST(FEpicRtcConferenceCreateTest, "System.Plugins.PixelStreaming2.FPS2EpicRtcConferenceCreateTest", EAutomationTestFlags::EditorContext | EAutomationTestFlags::ClientContext | EAutomationTestFlags::ProductFilter)
-	bool FEpicRtcConferenceCreateTest::RunTest(const FString& Parameters)
+	IMPLEMENT_SIMPLE_AUTOMATION_TEST(FPS2EpicRtcConferenceCreateTest, "System.Plugins.PixelStreaming2.FPS2EpicRtcConferenceCreateTest", EAutomationTestFlags::EditorContext | EAutomationTestFlags::ClientContext | EAutomationTestFlags::ProductFilter)
+	bool FPS2EpicRtcConferenceCreateTest::RunTest(const FString& Parameters)
 	{
 		EpicRtcErrorCode Result;
 
@@ -135,8 +135,8 @@ namespace UE::PixelStreaming2
 	}
 
 	// Tests the conference retrieval logic to ensure that a conference can be retrieved after is has been created
-	IMPLEMENT_SIMPLE_AUTOMATION_TEST(FEpicRtcConferenceGetTest, "System.Plugins.PixelStreaming2.FPS2EpicRtcConferenceGetTest", EAutomationTestFlags::EditorContext | EAutomationTestFlags::ClientContext | EAutomationTestFlags::ProductFilter)
-	bool FEpicRtcConferenceGetTest::RunTest(const FString& Parameters)
+	IMPLEMENT_SIMPLE_AUTOMATION_TEST(FPS2EpicRtcConferenceGetTest, "System.Plugins.PixelStreaming2.FPS2EpicRtcConferenceGetTest", EAutomationTestFlags::EditorContext | EAutomationTestFlags::ClientContext | EAutomationTestFlags::ProductFilter)
+	bool FPS2EpicRtcConferenceGetTest::RunTest(const FString& Parameters)
 	{
 		EpicRtcErrorCode Result;
 
diff --git a/Engine/Plugins/Media/PixelStreaming2/Source/PixelStreaming2/Private/Tests/EpicRtcPlatformTest.cpp b/Engine/Plugins/Media/PixelStreaming2/Source/PixelStreaming2/Private/Tests/EpicRtcPlatformTest.cpp
--- a/Engine/Plugins/Media/PixelStreaming2/Source/PixelStreaming2/Private/Tests/EpicRtcPlatformTest.cpp
+++ b/Engine/Plugins/Media/PixelStreaming2/Source/PixelStreaming2/Private/Tests/EpicRtcPlatformTest.cpp
@@ -11,8 +11,8 @@
 
 namespace UE::PixelStreaming2
 {
-	IMPLEMENT_SIMPLE_AUTOMATION_TEST(FEpicRtcPlatformCreateTest, "System.Plugins.PixelStreaming2.FPS2EpicRtcPlatformCreateTest", EAutomationTestFlags::EditorContext | EAutomationTestFlags::ClientContext | EAutomationTestFlags::ProductFilter)
-	bool FEpicRtcPlatformCreateTest::RunTest(const FString& Parameters)
+	IMPLEMENT_SIMPLE_AUTOMATION_TEST(FPS2EpicRtcPlatformCreateTest, "System.Plugins.PixelStreaming2.FPS2EpicRtcPlatformCreateTest", EAutomationTestFlags::EditorContext | EAutomationTestFlags::ClientContext | EAutomationTestFlags::ProductFilter)
+	bool FPS2EpicRtcPlatformCreateTest::RunTest(const FString& Parameters)
 	{
 		EpicRtcErrorCode Result;
 
@@ -26,8 +26,8 @@ namespace UE::PixelStreaming2
 		return true;
 	}
 
-	IMPLEMENT_SIMPLE_AUTOMATION_TEST(FEpicRtcPlatformLifetimeTest, "System.Plugins.PixelStreaming2.FPS2EpicRtcPlatformLifetimeTest", EAutomationTestFlags::EditorContext | EAutomationTestFlags::ClientContext | EAutomationTestFlags::ProductFilter)
-	bool FEpicRtcPlatformLifetimeTest::RunTest(const FString& Parameters)
+	IMPLEMENT_SIMPLE_AUTOMATION_TEST(FPS2EpicRtcPlatformLifetimeTest, "System.Plugins.PixelStreaming2.FPS2EpicRtcPlatformLifetimeTest", EAutomationTestFlags::EditorContext | EAutomationTestFlags::ClientContext | EAutomationTestFlags::ProductFilter)
+	bool FPS2EpicRtcPlatformLifetimeTest::RunTest(const FString& Parameters)
 	{
 		EpicRtcPlatformInterface* PlatformPtr;
 		uint32 CreationCount;
@@ -56,8 +56,8 @@ namespace UE::PixelStreaming2
 		return true;
 	}
 
-	IMPLEMENT_SIMPLE_AUTOMATION_TEST(FEpicRtcPlatformGetTest, "System.Plugins.PixelStreaming2.FPS2EpicRtcPlatformGetTest", EAutomationTestFlags::EditorContext | EAutomationTestFlags::ClientContext | EAutomationTestFlags::ProductFilter)
-	bool FEpicRtcPlatformGetTest::RunTest(const FString& Parameters)
+	IMPLEMENT_SIMPLE_AUTOMATION_TEST(FPS2EpicRtcPlatformGetTest, "System.Plugins.PixelStreaming2.FPS2EpicRtcPlatformGetTest", EAutomationTestFlags::EditorContext | EAutomationTestFlags::ClientContext | EAutomationTestFlags::ProductFilter)
+	bool FPS2EpicRtcPlatformGetTest::RunTest(const FString& Parameters)
 	{
 		EpicRtcErrorCode Result;
 
diff --git a/Engine/Plugins/Media/PixelStreaming2/Source/PixelStreaming2/Private/Tests/EpicRtcRoomTest.cpp b/Engine/Plugins/Media/PixelStreaming2/Source/PixelStreaming2/Private/Tests/EpicRtcRoomTest.cpp
--- a/Engine/Plugins/Media/PixelStreaming2/Source/PixelStreaming2/Private/Tests/EpicRtcRoomTest.cpp
+++ b/Engine/Plugins/Media/PixelStreaming2/Source/PixelStreaming2/Private/Tests/EpicRtcRoomTest.cpp
@@ -12,8 +12,8 @@
 namespace UE::PixelStreaming2
 {
 	// Tests the creation and removal of a session
-	IMPLEMENT_SIMPLE_AUTOMATION_TEST(FEpicRtcRoomLifetimeTest, "System.Plugins.PixelStreaming2.FPS2EpicRtcRoomLifetimeTest", EAutomationTestFlags::EditorContext | EAutomationTestFlags::ClientContext | EAutomationTestFlags::ProductFilter)
-	bool FEpicRtcRoomLifetimeTest::RunTest(const FString& Parameters)
+	IMPLEMENT_SIMPLE_AUTOMATION_TEST(FPS2EpicRtcRoomLifetimeTest, "System.Plugins.PixelStreaming2.FPS2EpicRtcRoomLifetimeTest", EAutomationTestFlags::EditorContext | EAutomationTestFlags::ClientContext | EAutomationTestFlags::ProductFilter)
+	bool FPS2EpicRtcRoomLifetimeTest::RunTest(const FString& Parameters)
 	{
 		EpicRtcErrorCode Result;
 
diff --git a/Engine/Plugins/Media/PixelStreaming2/Source/PixelStreaming2/Private/Tests/EpicRtcSessionTest.cpp b/Engine/Plugins/Media/PixelStreaming2/Source/PixelStreaming2/Private/Tests/EpicRtcSessionTest.cpp
--- a/Engine/Plugins/Media/PixelStreaming2/Source/PixelStreaming2/Private/Tests/EpicRtcSessionTest.cpp
+++ b/Engine/Plugins/Media/PixelStreaming2/Source/PixelStreaming2/Private/Tests/EpicRtcSessionTest.cpp
@@ -14,8 +14,8 @@
 namespace UE::PixelStreaming2
 {
 	// Tests the creation and removal of a session
-	IMPLEMENT_SIMPLE_AUTOMATION_TEST(FEpicRtcSessionLifetimeTest, "System.Plugins.PixelStreaming2.FPS2EpicRtcSessionLifetimeTest", EAutomationTestFlags::EditorContext | EAutomationTestFlags::ClientContext | EAutomationTestFlags::ProductFilter)
-	bool FEpicRtcSessionLifetimeTest::RunTest(const FString& Parameters)
+	IMPLEMENT_SIMPLE_AUTOMATION_TEST(FPS2EpicRtcSessionLifetimeTest, "System.Plugins.PixelStreaming2.FPS2EpicRtcSessionLifetimeTest", EAutomationTestFlags::EditorContext | EAutomationTestFlags::ClientContext | EAutomationTestFlags::ProductFilter)
+	bool FPS2EpicRtcSessionLifetimeTest::RunTest(const FString& Parameters)
 	{
 		EpicRtcErrorCode Result;
 
diff --git a/Engine/Plugins/Media/PixelStreaming2/Source/PixelStreaming2/Private/Tests/MultipleSignallingConnectionsTest.cpp b/Engine/Plugins/Media/PixelStreaming2/Source/PixelStreaming2/Private/Tests/MultipleSignallingConnectionsTest.cpp
--- a/Engine/Plugins/Media/PixelStreaming2/Source/PixelStreaming2/Private/Tests/MultipleSignallingConnectionsTest.cpp
+++ b/Engine/Plugins/Media/PixelStreaming2/Source/PixelStreaming2/Private/Tests/MultipleSignallingConnectionsTest.cpp
@@ -162,8 +162,8 @@ namespace UE::PixelStreaming2
 		return false;
 	}
 
-	IMPLEMENT_SIMPLE_AUTOMATION_TEST(FMultipleSignallingConnectionsTest, "System.Plugins.PixelStreaming2.MultipleSignallingConnectionsTest", EAutomationTestFlags::EditorContext | EAutomationTestFlags::ClientContext | EAutomationTestFlags::ProductFilter);
-	bool FMultipleSignallingConnectionsTest::RunTest(const FString& Parameters)
+	IMPLEMENT_SIMPLE_AUTOMATION_TEST(FPS2MultipleSignallingConnectionsTest, "System.Plugins.PixelStreaming2.MultipleSignallingConnectionsTest", EAutomationTestFlags::EditorContext | EAutomationTestFlags::ClientContext | EAutomationTestFlags::ProductFilter);
+	bool FPS2MultipleSignallingConnectionsTest::RunTest(const FString& Parameters)
 	{
 		UE_LOG(LogPixelStreaming2, Log, TEXT("----------- ConnectAndDisconnectMultipleStreamersEmbeddedCirrus -----------"));
 
diff --git a/Engine/Plugins/Media/PixelStreaming2/Source/PixelStreaming2/Private/Tests/ProtocolTest.cpp b/Engine/Plugins/Media/PixelStreaming2/Source/PixelStreaming2/Private/Tests/ProtocolTest.cpp
--- a/Engine/Plugins/Media/PixelStreaming2/Source/PixelStreaming2/Private/Tests/ProtocolTest.cpp
+++ b/Engine/Plugins/Media/PixelStreaming2/Source/PixelStreaming2/Private/Tests/ProtocolTest.cpp
@@ -15,8 +15,8 @@
 
 namespace UE::PixelStreaming2
 {
-	IMPLEMENT_SIMPLE_AUTOMATION_TEST(FProtocolTestAddMessage, "System.Plugins.PixelStreaming2.FPS2ProtocolTestAddMessage", EAutomationTestFlags::EditorContext | EAutomationTestFlags::ClientContext | EAutomationTestFlags::ProductFilter)
-	bool FProtocolTestAddMessage::RunTest(const FString& Parameters)
+	IMPLEMENT_SIMPLE_AUTOMATION_TEST(FPS2ProtocolTestAddMessage, "System.Plugins.PixelStreaming2.FPS2ProtocolTestAddMessage", EAutomationTestFlags::EditorContext | EAutomationTestFlags::ClientContext | EAutomationTestFlags::ProductFilter)
+	bool FPS2ProtocolTestAddMessage::RunTest(const FString& Parameters)
 	{
 		// need to be able to accept codec to handshake otherwise setting local description fails when generating an answer
 		SetCodec(EVideoCodec::VP8);
@@ -93,8 +93,8 @@ namespace UE::PixelStreaming2
 		return true;
 	}
 
-	IMPLEMENT_SIMPLE_AUTOMATION_TEST(FProtocolTestUseCustomMessage, "System.Plugins.PixelStreaming2.FPS2ProtocolTestUseCustomMessage", EAutomationTestFlags::EditorContext | EAutomationTestFlags::ClientContext | EAutomationTestFlags::ProductFilter)
-	bool FProtocolTestUseCustomMessage::RunTest(const FString& Parameters)
+	IMPLEMENT_SIMPLE_AUTOMATION_TEST(FPS2ProtocolTestUseCustomMessage, "System.Plugins.PixelStreaming2.FPS2ProtocolTestUseCustomMessage", EAutomationTestFlags::EditorContext | EAutomationTestFlags::ClientContext | EAutomationTestFlags::ProductFilter)
+	bool FPS2ProtocolTestUseCustomMessage::RunTest(const FString& Parameters)
 	{
 		// need to be able to accept codec to handshake otherwise setting local description fails when generating an answer
 		SetCodec(EVideoCodec::VP8);
diff --git a/Engine/Plugins/Media/PixelStreaming2/Source/PixelStreaming2/Private/Tests/VP8Tests.cpp b/Engine/Plugins/Media/PixelStreaming2/Source/PixelStreaming2/Private/Tests/VP8Tests.cpp
--- a/Engine/Plugins/Media/PixelStreaming2/Source/PixelStreaming2/Private/Tests/VP8Tests.cpp
+++ b/Engine/Plugins/Media/PixelStreaming2/Source/PixelStreaming2/Private/Tests/VP8Tests.cpp
@@ -7,16 +7,16 @@
 
 namespace UE::PixelStreaming2
 {
-	IMPLEMENT_SIMPLE_AUTOMATION_TEST(FVP8FrameReceivedTest, "System.Plugins.PixelStreaming2.FPS2VP8FrameReceivedTest", EAutomationTestFlags::EditorContext | EAutomationTestFlags::ClientContext | EAutomationTestFlags::ProductFilter)
-	bool FVP8FrameReceivedTest::RunTest(const FString& Parameters)
+	IMPLEMENT_SIMPLE_AUTOMATION_TEST(FPS2VP8FrameReceivedTest, "System.Plugins.PixelStreaming2.FPS2VP8FrameReceivedTest", EAutomationTestFlags::EditorContext | EAutomationTestFlags::ClientContext | EAutomationTestFlags::ProductFilter)
+	bool FPS2VP8FrameReceivedTest::RunTest(const FString& Parameters)
 	{
 		SetCodec(EVideoCodec::VP8);
 		DoFrameReceiveTest();
 		return true;
 	}
 
-	IMPLEMENT_SIMPLE_AUTOMATION_TEST(FVP8FrameResizeTest, "System.Plugins.PixelStreaming2.FPS2VP8FrameResizeTest", EAutomationTestFlags::EditorContext | EAutomationTestFlags::ClientContext | EAutomationTestFlags::ProductFilter)
-	bool FVP8FrameResizeTest::RunTest(const FString& Parameters)
+	IMPLEMENT_SIMPLE_AUTOMATION_TEST(FPS2VP8FrameResizeTest, "System.Plugins.PixelStreaming2.FPS2VP8FrameResizeTest", EAutomationTestFlags::EditorContext | EAutomationTestFlags::ClientContext | EAutomationTestFlags::ProductFilter)
+	bool FPS2VP8FrameResizeTest::RunTest(const FString& Parameters)
 	{
 		SetCodec(EVideoCodec::VP8);
 		DoFrameResizeMultipleTimesTest();
diff --git a/Engine/Plugins/Media/PixelStreaming2/Source/PixelStreaming2/Private/Tests/VP9Tests.cpp b/Engine/Plugins/Media/PixelStreaming2/Source/PixelStreaming2/Private/Tests/VP9Tests.cpp
--- a/Engine/Plugins/Media/PixelStreaming2/Source/PixelStreaming2/Private/Tests/VP9Tests.cpp
+++ b/Engine/Plugins/Media/PixelStreaming2/Source/PixelStreaming2/Private/Tests/VP9Tests.cpp
@@ -7,16 +7,16 @@
 
 namespace UE::PixelStreaming2
 {
-	IMPLEMENT_SIMPLE_AUTOMATION_TEST(FVP9FrameReceivedTest, "System.Plugins.PixelStreaming2.FPS2VP9FrameReceivedTest", EAutomationTestFlags::EditorContext | EAutomationTestFlags::ClientContext | EAutomationTestFlags::ProductFilter)
-	bool FVP9FrameReceivedTest::RunTest(const FString& Parameters)
+	IMPLEMENT_SIMPLE_AUTOMATION_TEST(FPS2VP9FrameReceivedTest, "System.Plugins.PixelStreaming2.FPS2VP9FrameReceivedTest", EAutomationTestFlags::EditorContext | EAutomationTestFlags::ClientContext | EAutomationTestFlags::ProductFilter)
+	bool FPS2VP9FrameReceivedTest::RunTest(const FString& Parameters)
 	{
 		SetCodec(EVideoCodec::VP9);
 		DoFrameReceiveTest();
 		return true;
 	}
 
-	IMPLEMENT_SIMPLE_AUTOMATION_TEST(FVP9FrameResizeTest, "System.Plugins.PixelStreaming2.FPS2VP9FrameResizeTest", EAutomationTestFlags::EditorContext | EAutomationTestFlags::ClientContext | EAutomationTestFlags::ProductFilter)
-	bool FVP9FrameResizeTest::RunTest(const FString& Parameters)
+	IMPLEMENT_SIMPLE_AUTOMATION_TEST(FPS2VP9FrameResizeTest, "System.Plugins.PixelStreaming2.FPS2VP9FrameResizeTest", EAutomationTestFlags::EditorContext | EAutomationTestFlags::ClientContext | EAutomationTestFlags::ProductFilter)
+	bool FPS2VP9FrameResizeTest::RunTest(const FString& Parameters)
 	{
 		SetCodec(EVideoCodec::VP9);
 		DoFrameResizeMultipleTimesTest();
diff --git a/Engine/Plugins/Media/PixelStreaming2/Source/PixelStreaming2Core/Internal/PixelStreaming2PluginSettings.h b/Engine/Plugins/Media/PixelStreaming2/Source/PixelStreaming2Core/Internal/PixelStreaming2PluginSettings.h
--- a/Engine/Plugins/Media/PixelStreaming2/Source/PixelStreaming2Core/Internal/PixelStreaming2PluginSettings.h
+++ b/Engine/Plugins/Media/PixelStreaming2/Source/PixelStreaming2Core/Internal/PixelStreaming2PluginSettings.h
@@ -357,8 +357,7 @@ public:
 		DisplayName = "Disable Audio/Video track sychronisation",
 		ToolTip = "Disables the synchronization of audio and video tracks in WebRTC. This can be useful in low latency usecases where synchronization is not required."
 		))
-	// TODO (Migration): Set to true once RTCP-6584 is fixed
-	bool WebRTCDisableAudioSync = false;
+	bool WebRTCDisableAudioSync = true;
 
 	static TAutoConsoleVariable<bool> CVarWebRTCEnableFlexFec;
 	UPROPERTY(config, EditAnywhere, Category = "WebRtcSettings", meta = (
diff --git a/Engine/Plugins/Media/PixelStreaming2/Source/PixelStreaming2Servers/Private/Tests/ServerTests.cpp b/Engine/Plugins/Media/PixelStreaming2/Source/PixelStreaming2Servers/Private/Tests/ServerTests.cpp
--- a/Engine/Plugins/Media/PixelStreaming2/Source/PixelStreaming2Servers/Private/Tests/ServerTests.cpp
+++ b/Engine/Plugins/Media/PixelStreaming2/Source/PixelStreaming2Servers/Private/Tests/ServerTests.cpp
@@ -135,8 +135,8 @@ namespace UE::PixelStreaming2Servers
 		}
 	}
 
-	IMPLEMENT_SIMPLE_AUTOMATION_TEST(FLaunchEmbeddedCirrusTest, "System.Plugins.PixelStreaming2.LaunchEmbeddedCirrus", EAutomationTestFlags::EditorContext | EAutomationTestFlags::ClientContext | EAutomationTestFlags::ProductFilter)
-	bool FLaunchEmbeddedCirrusTest::RunTest(const FString& Parameters)
+	IMPLEMENT_SIMPLE_AUTOMATION_TEST(FPS2LaunchEmbeddedCirrus, "System.Plugins.PixelStreaming2.FPS2LaunchEmbeddedCirrus", EAutomationTestFlags::EditorContext | EAutomationTestFlags::ClientContext | EAutomationTestFlags::ProductFilter)
+	bool FPS2LaunchEmbeddedCirrus::RunTest(const FString& Parameters)
 	{
 		UE_LOG(LogPixelStreaming2Servers, Log, TEXT("----------- LaunchEmbeddedCirrusTest -----------"));
 
@@ -185,8 +185,8 @@ namespace UE::PixelStreaming2Servers
 		return true;
 	}
 
-	IMPLEMENT_SIMPLE_AUTOMATION_TEST(FTwoWebsocketToEmbeddedCirrusTest, "System.Plugins.PixelStreaming2.TwoWebsocketToEmbeddedCirrus", EAutomationTestFlags::EditorContext | EAutomationTestFlags::ClientContext | EAutomationTestFlags::ProductFilter)
-	bool FTwoWebsocketToEmbeddedCirrusTest::RunTest(const FString& Parameters)
+	IMPLEMENT_SIMPLE_AUTOMATION_TEST(FPS2TwoWebsocketToEmbeddedCirrusTest, "System.Plugins.PixelStreaming2.FPS2TwoWebsocketToEmbeddedCirrusTest", EAutomationTestFlags::EditorContext | EAutomationTestFlags::ClientContext | EAutomationTestFlags::ProductFilter)
+	bool FPS2TwoWebsocketToEmbeddedCirrusTest::RunTest(const FString& Parameters)
 	{
 		UE_LOG(LogPixelStreaming2Servers, Log, TEXT("----------- TwoWebsocketToEmbeddedCirrus -----------"));
 
