[PS2] Pixel Streaming 2 only send websocket ping after LastKeepAliveCycles instead of when first connected.

This fix sets LastKeepAliveCycles on the first connection. This fixes sending ping immediately which can be received before a room is constructed causing a warning message to be received in the PS2 tests.

diff --git a/Engine/Plugins/Media/PixelStreaming2/Source/PixelStreaming2/Private/EpicRtcWebsocket.cpp b/Engine/Plugins/Media/PixelStreaming2/Source/PixelStreaming2/Private/EpicRtcWebsocket.cpp
--- a/Engine/Plugins/Media/PixelStreaming2/Source/PixelStreaming2/Private/EpicRtcWebsocket.cpp
+++ b/Engine/Plugins/Media/PixelStreaming2/Source/PixelStreaming2/Private/EpicRtcWebsocket.cpp
@@ -121,6 +121,7 @@ namespace UE::PixelStreaming2
 		UE_LOG(LogEpicRtcWebsocket, Log, TEXT("Websocket connection made to: %s"), *Url);
 		bCloseRequested = false;
 		bReconnectOnError = false;
+		LastKeepAliveCycles = FPlatformTime::Cycles64();
 		Observer->OnOpen();
 	}
 
