[PS2] Add log category for WebRtc used from EpicRtc.

WebRtc likes to pring error and warnings to the log for messages that should not be errors. But EpicRtc als oprints out messages that we want to catch. This  change adds a log category to split EpicRtc adn WebRtc messages.

diff --git a/Engine/Plugins/Media/PixelStreaming2/Source/PixelStreaming2/Internal/EpicRtcLogging.h b/Engine/Plugins/Media/PixelStreaming2/Source/PixelStreaming2/Internal/EpicRtcLogging.h
--- a/Engine/Plugins/Media/PixelStreaming2/Source/PixelStreaming2/Internal/EpicRtcLogging.h
+++ b/Engine/Plugins/Media/PixelStreaming2/Source/PixelStreaming2/Internal/EpicRtcLogging.h
@@ -6,7 +6,8 @@
 
 #include "epic_rtc/common/logging.h"
 
-DECLARE_LOG_CATEGORY_EXTERN(LogPixelStreaming2EpicRtc, NoLogging, All);
+DECLARE_LOG_CATEGORY_EXTERN(LogPixelStreaming2EpicRtc, Log, All);
+DECLARE_LOG_CATEGORY_EXTERN(LogPixelStreaming2WebRtc, Fatal, All);
 
 namespace UE::PixelStreaming2
 {
diff --git a/Engine/Plugins/Media/PixelStreaming2/Source/PixelStreaming2/Private/EpicRtcLogging.cpp b/Engine/Plugins/Media/PixelStreaming2/Source/PixelStreaming2/Private/EpicRtcLogging.cpp
--- a/Engine/Plugins/Media/PixelStreaming2/Source/PixelStreaming2/Private/EpicRtcLogging.cpp
+++ b/Engine/Plugins/Media/PixelStreaming2/Source/PixelStreaming2/Private/EpicRtcLogging.cpp
@@ -3,6 +3,7 @@
 #include "EpicRtcLogging.h"
 
 DEFINE_LOG_CATEGORY(LogPixelStreaming2EpicRtc);
+DEFINE_LOG_CATEGORY(LogPixelStreaming2WebRtc);
 
 namespace UE::PixelStreaming2
 {
diff --git a/Engine/Plugins/Media/PixelStreaming2/Source/PixelStreaming2/Private/PixelStreaming2Module.cpp b/Engine/Plugins/Media/PixelStreaming2/Source/PixelStreaming2/Private/PixelStreaming2Module.cpp
--- a/Engine/Plugins/Media/PixelStreaming2/Source/PixelStreaming2/Private/PixelStreaming2Module.cpp
+++ b/Engine/Plugins/Media/PixelStreaming2/Source/PixelStreaming2/Private/PixelStreaming2Module.cpp
@@ -710,7 +710,7 @@ namespace UE::PixelStreaming2
 				._logger = new FEpicRtcLogsRedirector(),
 #if !NO_LOGGING // When building WITH_SHIPPING by default .GetVerbosity() does not exist
 				._level = UnrealLogToEpicRtcCategoryMap[LogPixelStreaming2EpicRtc.GetVerbosity()],
-				._levelWebRtc = UnrealLogToEpicRtcCategoryMap[LogPixelStreaming2EpicRtc.GetVerbosity()]
+				._levelWebRtc = UnrealLogToEpicRtcCategoryMap[LogPixelStreaming2WebRtc.GetVerbosity()]
 #endif
 			},
 			._stats = {
