[PS2] Fix resize tests sporadically failing.

Resize test could fail when the last frame arrives late after the size is checked in the test, causing the last frame's size to be compared to the expected size which will be wrong.

This fix checks against the last received frame rather than the first frame.

diff --git a/Engine/Plugins/Media/PixelStreaming2/Source/PixelStreaming2/Private/Tests/MockPlayer.cpp b/Engine/Plugins/Media/PixelStreaming2/Source/PixelStreaming2/Private/Tests/MockPlayer.cpp
--- a/Engine/Plugins/Media/PixelStreaming2/Source/PixelStreaming2/Private/Tests/MockPlayer.cpp
+++ b/Engine/Plugins/Media/PixelStreaming2/Source/PixelStreaming2/Private/Tests/MockPlayer.cpp
@@ -19,11 +19,8 @@ namespace UE::PixelStreaming2
 
 	void FMockVideoSink::OnFrame(const EpicRtcVideoFrame& Frame)
 	{
-		if (!bReceivedFrame)
-		{
-			VideoBuffer = Frame._buffer;
-			bReceivedFrame = true;
-		}
+		VideoBuffer = Frame._buffer;
+		bReceivedFrame = true;
 	}
 
 	void FMockVideoSink::ResetReceivedFrame()
