[PS2] Fix Pixel Streaming delegate data race.

Several Delegates are removed while broadcast is called on a separate thread, which causes a ensure to fail that checks for the data race. These data races happen in the PixelStreaming2 MockPlayer unit test when when ending the stream.

 This fix calls the broadcast in the game thread.

diff --git a/Engine/Plugins/Media/PixelStreaming2/Source/PixelStreaming2/Private/Tests/MockPlayer.cpp b/Engine/Plugins/Media/PixelStreaming2/Source/PixelStreaming2/Private/Tests/MockPlayer.cpp
--- a/Engine/Plugins/Media/PixelStreaming2/Source/PixelStreaming2/Private/Tests/MockPlayer.cpp
+++ b/Engine/Plugins/Media/PixelStreaming2/Source/PixelStreaming2/Private/Tests/MockPlayer.cpp
@@ -2,6 +2,7 @@
 
 #include "MockPlayer.h"
 
+#include "Async/Async.h"
 #include "DefaultDataProtocol.h"
 #include "EpicRtcVideoEncoderInitializer.h"
 #include "EpicRtcVideoDecoderInitializer.h"
@@ -252,9 +253,16 @@ namespace UE::PixelStreaming2
 			return;
 		}
 
-		const TArray<uint8> Data(DataFrame->Data(), DataFrame->Size());
-
-		OnMessageReceived.Broadcast(Data);
+		// Broadcast must be done on the GameThread because the GameThread can remove the delegates.
+		// If removing and broadcast happens simultaneously it causes a datarace failure.
+		TWeakPtr<FMockPlayer> WeakPlayer = AsShared();
+		AsyncTask(ENamedThreads::GameThread, [WeakPlayer, DataFrame]() {
+			if (TSharedPtr<FMockPlayer> PinnedPlayer = WeakPlayer.Pin())
+			{
+				const TArray<uint8> Data(DataFrame->Data(), DataFrame->Size());
+				PinnedPlayer->OnMessageReceived.Broadcast(Data);
+			}
+		});
 	}
 
 	void FMockPlayer::OnDataTrackRemoved(EpicRtcDataTrackInterface*) {}
diff --git a/Engine/Plugins/Media/PixelStreaming2/Source/PixelStreaming2/Private/Tests/MockPlayer.h b/Engine/Plugins/Media/PixelStreaming2/Source/PixelStreaming2/Private/Tests/MockPlayer.h
--- a/Engine/Plugins/Media/PixelStreaming2/Source/PixelStreaming2/Private/Tests/MockPlayer.h
+++ b/Engine/Plugins/Media/PixelStreaming2/Source/PixelStreaming2/Private/Tests/MockPlayer.h
@@ -39,7 +39,7 @@ namespace UE::PixelStreaming2
 		FThreadSafeBool							  bReceivedFrame = false;
 	};
 
-	class FMockPlayer
+	class FMockPlayer : public TSharedFromThis<FMockPlayer>
 	{
 	public:
 		FMockPlayer();
diff --git a/Engine/Plugins/Media/PixelStreaming2/Source/PixelStreaming2/Private/Tests/TestUtils.h b/Engine/Plugins/Media/PixelStreaming2/Source/PixelStreaming2/Private/Tests/TestUtils.h
--- a/Engine/Plugins/Media/PixelStreaming2/Source/PixelStreaming2/Private/Tests/TestUtils.h
+++ b/Engine/Plugins/Media/PixelStreaming2/Source/PixelStreaming2/Private/Tests/TestUtils.h
@@ -29,8 +29,8 @@ namespace UE::PixelStreaming2
 			, Callback(InCallback)
 			, bComplete(InbComplete)
 		{
-			MessageReceivedHandle = Player->OnMessageReceived.AddLambda([this](const TArray<uint8>& RawBuffer) {
-				(Callback)(RawBuffer);
+			MessageReceivedHandle = Player->OnMessageReceived.AddLambda([MessageCallback = this->Callback](const TArray<uint8>& RawBuffer) {
+				MessageCallback(RawBuffer);
 			});
 		}
 
