[PS2] Fix Pixel Streaming 2 capturing with multiple FrameCapturers. These could cause a data race inside delegates in Pixel Capture in OnComplete.

The fix is to reset the FrameCapturer when the streamer starts a new stream.

diff --git a/Engine/Plugins/Media/PixelStreaming2/Source/PixelStreaming2/Internal/VideoCapturer.h b/Engine/Plugins/Media/PixelStreaming2/Source/PixelStreaming2/Internal/VideoCapturer.h
--- a/Engine/Plugins/Media/PixelStreaming2/Source/PixelStreaming2/Internal/VideoCapturer.h
+++ b/Engine/Plugins/Media/PixelStreaming2/Source/PixelStreaming2/Internal/VideoCapturer.h
@@ -35,6 +35,8 @@ namespace UE::PixelStreaming2
 
 		TSharedPtr<IPixelCaptureOutputFrame> RequestFormat(int32 Format, int32 LayerIndex = -1);
 
+		void ResetFrameCapturer();
+
 		/**
 		 * This is broadcast each time a frame exits the adapt process. Used to synchronize framerates with input rates.
 		 * This should be called once per frame taking into consideration all the target formats and layers within the frame.
diff --git a/Engine/Plugins/Media/PixelStreaming2/Source/PixelStreaming2/Private/Streamer.cpp b/Engine/Plugins/Media/PixelStreaming2/Source/PixelStreaming2/Private/Streamer.cpp
--- a/Engine/Plugins/Media/PixelStreaming2/Source/PixelStreaming2/Private/Streamer.cpp
+++ b/Engine/Plugins/Media/PixelStreaming2/Source/PixelStreaming2/Private/Streamer.cpp
@@ -165,6 +165,8 @@ namespace UE::PixelStreaming2
 			AllConnectionsClosedHandle = Delegates->OnAllConnectionsClosedNative.AddSP(AsShared(), &FStreamer::TriggerMouseLeave);
 		}
 
+		VideoCapturer->ResetFrameCapturer();
+
 		// Broadcast the preconnection event just before we do `TryConnect`
 		StreamingPreConnectionEvent.Broadcast(this);
 
diff --git a/Engine/Plugins/Media/PixelStreaming2/Source/PixelStreaming2/Private/VideoCapturer.cpp b/Engine/Plugins/Media/PixelStreaming2/Source/PixelStreaming2/Private/VideoCapturer.cpp
--- a/Engine/Plugins/Media/PixelStreaming2/Source/PixelStreaming2/Private/VideoCapturer.cpp
+++ b/Engine/Plugins/Media/PixelStreaming2/Source/PixelStreaming2/Private/VideoCapturer.cpp
@@ -103,6 +103,11 @@ namespace UE::PixelStreaming2
 		return nullptr;
 	}
 
+	void FVideoCapturer::ResetFrameCapturer()
+	{
+		CreateFrameCapturer();
+	}
+
 	void FVideoCapturer::OnSimulcastEnabledChanged(IConsoleVariable* Var)
 	{
 		CreateFrameCapturer();
