[PS2] EpicRtc had s a bug where it incorrectly calculates stride for frames with odd resolution fixed by RTCP-7246.

This change removes the hack to handle this.

diff --git a/Engine/Plugins/Media/PixelStreaming2/Source/PixelStreaming2/Private/EpicRtcVideoDecoder.cpp b/Engine/Plugins/Media/PixelStreaming2/Source/PixelStreaming2/Private/EpicRtcVideoDecoder.cpp
--- a/Engine/Plugins/Media/PixelStreaming2/Source/PixelStreaming2/Private/EpicRtcVideoDecoder.cpp
+++ b/Engine/Plugins/Media/PixelStreaming2/Source/PixelStreaming2/Private/EpicRtcVideoDecoder.cpp
@@ -131,24 +131,22 @@ namespace UE::PixelStreaming2
 
 				if (DecodeResult.IsSuccess())
 				{
-					// TODO(Eden.Harris) RTCP-7247 EpicRtc currently has a bug where it incorrectly calculates stride for frames with odd resolution.
-					// To handle this, round width/height to a even number. When EpicRtc is fixed by RTCP-7246, this hack can be removed.
-					uint32 RoundedWidth = DecoderResource->GetWidth() & ~1;
-					uint32 RoundedHeight = DecoderResource->GetHeight() & ~1;
+					const uint32 FrameWidth = DecoderResource->GetWidth();
+					const uint32 FrameHeight = DecoderResource->GetHeight();
 					
-					TSharedPtr<FPixelCaptureBufferI420> I420Buffer = MakeShared<FPixelCaptureBufferI420>(RoundedWidth, RoundedHeight);
+					TSharedPtr<FPixelCaptureBufferI420> I420Buffer = MakeShared<FPixelCaptureBufferI420>(FrameWidth, FrameHeight);
 
-					uint32 DataSizeY = DecoderResource->GetWidth() * DecoderResource->GetHeight();
-					uint32 DataSizeUV = ((DecoderResource->GetWidth() + 1) / 2) * ((DecoderResource->GetHeight() + 1) / 2);
+					const uint32 DataSizeY = FrameWidth * FrameHeight;
+					const uint32 DataSizeUV = ((FrameWidth + 1) / 2) * ((FrameHeight + 1) / 2);
 
 					CopyI420(
-						DecoderResource->GetRaw().Get(), DecoderResource->GetWidth(),
-						DecoderResource->GetRaw().Get() + DataSizeY, (DecoderResource->GetWidth() + 1) / 2,
-						DecoderResource->GetRaw().Get() + DataSizeY + DataSizeUV, (DecoderResource->GetWidth() + 1) / 2,
+						DecoderResource->GetRaw().Get(), FrameWidth,
+						DecoderResource->GetRaw().Get() + DataSizeY, (FrameWidth + 1) / 2,
+						DecoderResource->GetRaw().Get() + DataSizeY + DataSizeUV, (FrameWidth + 1) / 2,
 						I420Buffer->GetMutableDataY(), I420Buffer->GetStrideY(),
 						I420Buffer->GetMutableDataU(), I420Buffer->GetStrideUV(),
 						I420Buffer->GetMutableDataV(), I420Buffer->GetStrideUV(),
-						RoundedWidth, RoundedHeight);
+						FrameWidth, FrameHeight);
 
 					DecodedFrame._buffer = new FEpicRtcVideoBufferI420(I420Buffer);
 				}
