[PS2] This change fixes PixelStreamingWebRTCPortAllocatorFlags parameters not being read correctly.

Problem was SetCVarFromPropertyAndValue would set the value as a integer, and not the string.

diff --git a/Engine/Plugins/Media/PixelStreaming2/Source/PixelStreaming2Core/Private/PixelStreaming2PluginSettings.cpp b/Engine/Plugins/Media/PixelStreaming2/Source/PixelStreaming2Core/Private/PixelStreaming2PluginSettings.cpp
--- a/Engine/Plugins/Media/PixelStreaming2/Source/PixelStreaming2Core/Private/PixelStreaming2PluginSettings.cpp
+++ b/Engine/Plugins/Media/PixelStreaming2/Source/PixelStreaming2Core/Private/PixelStreaming2PluginSettings.cpp
@@ -1150,7 +1150,7 @@ void UPixelStreaming2PluginSettings::PostInitProperties()
 				FString ConsoleString;
 				if (FParse::Value(FCommandLine::Get(), *ConsoleVariableToCommandArgValue(CVarString), ConsoleString))
 				{
-					SetCVarFromPropertyAndValue(IConsoleManager::Get().FindConsoleVariable(*CVarString), Property, CVarString, ConsoleString);
+					IConsoleManager::Get().FindConsoleVariable(*CVarString)->Set(*ConsoleString, ECVF_SetByCommandline);
 				}
 				else
 				{
