[PS2] Fix incorrect Quality set from QP.

Pixel Streaming 2 was setting incorrect Quality from QP. The Min and Max values were reversed and floating point math caused values to be below zero.

diff --git a/Engine/Plugins/Media/PixelStreaming2/Source/PixelStreaming2/Private/PixelStreaming2Module.cpp b/Engine/Plugins/Media/PixelStreaming2/Source/PixelStreaming2/Private/PixelStreaming2Module.cpp
--- a/Engine/Plugins/Media/PixelStreaming2/Source/PixelStreaming2/Private/PixelStreaming2Module.cpp
+++ b/Engine/Plugins/Media/PixelStreaming2/Source/PixelStreaming2/Private/PixelStreaming2Module.cpp
@@ -533,13 +533,13 @@ namespace UE::PixelStreaming2
 			// Set Encoder.MinQP Legacy CVar
 			InputHandler->SetCommandHandler(TEXT("Encoder.MinQP"), [](FString PlayerId, FString Descriptor, FString MinQPString) {
 				int MinQP = FCString::Atoi(*MinQPString);
-				UPixelStreaming2PluginSettings::CVarEncoderMinQuality->Set(100.0f * (1.0f - (FMath::Clamp<int32>(MinQP, 0, 51) / 51.0f)), ECVF_SetByCommandline);
+				UPixelStreaming2PluginSettings::CVarEncoderMaxQuality->Set(FMath::RoundToFloat(100.0f * (1.0f - (FMath::Clamp<int32>(MinQP, 0, 51) / 51.0f))), ECVF_SetByCommandline);
 			});
 
 			// Set Encoder.MaxQP Legacy CVar
 			InputHandler->SetCommandHandler(TEXT("Encoder.MaxQP"), [](FString PlayerId, FString Descriptor, FString MaxQPString) {
 				int MaxQP = FCString::Atoi(*MaxQPString);
-				UPixelStreaming2PluginSettings::CVarEncoderMaxQuality->Set(100.0f * (1.0f - (FMath::Clamp<int32>(MaxQP, 0, 51) / 51.0f)), ECVF_SetByCommandline);
+				UPixelStreaming2PluginSettings::CVarEncoderMinQuality->Set(FMath::RoundToFloat(100.0f * (1.0f - (FMath::Clamp<int32>(MaxQP, 0, 51) / 51.0f))), ECVF_SetByCommandline);
 			});
 
 			// Set Encoder.MinQuality CVar
