[PS2] Fix frames appearing out of order.

Problem is fixed by using NoCopy capturer which works with the MediaCapturer which has the frames in the same format. By having the same format, no extra copy is needed which could cause sync issues when no fence was used.

diff --git a/Engine/Plugins/Media/PixelStreaming2/Source/PixelStreaming2/Internal/VideoProducer.h b/Engine/Plugins/Media/PixelStreaming2/Source/PixelStreaming2/Internal/VideoProducer.h
--- a/Engine/Plugins/Media/PixelStreaming2/Source/PixelStreaming2/Internal/VideoProducer.h
+++ b/Engine/Plugins/Media/PixelStreaming2/Source/PixelStreaming2/Internal/VideoProducer.h
@@ -21,6 +21,8 @@ namespace UE::PixelStreaming2
 
 		virtual FString ToString() override;
 
+		virtual bool IsFrameAlreadyCopied() { return false; }
+
 	protected:
 		FVideoProducer() = default;
 	};
diff --git a/Engine/Plugins/Media/PixelStreaming2/Source/PixelStreaming2/Internal/VideoProducerMediaCapture.h b/Engine/Plugins/Media/PixelStreaming2/Source/PixelStreaming2/Internal/VideoProducerMediaCapture.h
--- a/Engine/Plugins/Media/PixelStreaming2/Source/PixelStreaming2/Internal/VideoProducerMediaCapture.h
+++ b/Engine/Plugins/Media/PixelStreaming2/Source/PixelStreaming2/Internal/VideoProducerMediaCapture.h
@@ -54,6 +54,8 @@ namespace UE::PixelStreaming2
 
 		virtual FString ToString() override;
 
+		virtual bool IsFrameAlreadyCopied() override { return true; }
+
 	protected:
 		void StartActiveViewportCapture();
 		void LateStartActiveViewportCapture();
@@ -62,7 +64,6 @@ namespace UE::PixelStreaming2
 		FVideoProducerMediaCapture() = default;
 		FVideoProducerMediaCapture(TObjectPtr<UPixelStreaming2MediaIOCapture> MediaCapture);
 
-		void CreateDefaultActiveViewportMediaCapture();
 		void OnCaptureActiveViewportStateChanged();
 
 	private:
diff --git a/Engine/Plugins/Media/PixelStreaming2/Source/PixelStreaming2/Private/VideoCapturer.cpp b/Engine/Plugins/Media/PixelStreaming2/Source/PixelStreaming2/Private/VideoCapturer.cpp
--- a/Engine/Plugins/Media/PixelStreaming2/Source/PixelStreaming2/Private/VideoCapturer.cpp
+++ b/Engine/Plugins/Media/PixelStreaming2/Source/PixelStreaming2/Private/VideoCapturer.cpp
@@ -9,6 +9,7 @@
 #include "PixelCaptureCapturerI420ToRHI.h"
 #include "PixelCaptureCapturerNV12ToRHI.h"
 #include "PixelCaptureCapturerRHI.h"
+#include "PixelCaptureCapturerRHINoCopy.h"
 #include "PixelCaptureCapturerRHIRDG.h"
 #include "PixelCaptureCapturerRHIToI420CPU.h"
 #include "PixelStreaming2PluginSettings.h"
@@ -150,13 +151,22 @@ namespace UE::PixelStreaming2
 			switch (FinalFormat)
 			{
 				case PixelCaptureBufferFormat::FORMAT_RHI:
-					if(UPixelStreaming2PluginSettings::CVarCaptureUseFence.GetValueOnAnyThread())
+					if (VideoProducer->IsFrameAlreadyCopied() && FinalScale == 1.0 && UE::PixelStreaming2::GetSimulcastParameters().Num() == 1)
 					{
-						return FPixelCaptureCapturerRHI::Create(FinalScale);
+						// If we only have a single layer (and it's scale is 1), we can use the no copy capturer
+						// as we know the output from the media capture will already be the correct format and scale
+						return FPixelCaptureCapturerRHINoCopy::Create(FinalScale);
 					}
-					else 
+					else
 					{
-						return FPixelCaptureCapturerRHIRDG::Create(FinalScale);
+						if (UPixelStreaming2PluginSettings::CVarCaptureUseFence.GetValueOnAnyThread())
+						{
+							return FPixelCaptureCapturerRHI::Create(FinalScale);
+						}
+						else
+						{
+							return FPixelCaptureCapturerRHIRDG::Create(FinalScale);
+						}
 					}
 				case PixelCaptureBufferFormat::FORMAT_I420:
 					return FPixelCaptureCapturerRHIToI420CPU::Create(FinalScale);
