[PS2] Fix frames appearing out of order in Pixel Streaming.

This fix defaults to using a fence to capture frames. It can result in lower capture rate.

diff --git a/Engine/Plugins/Media/PixelStreaming2/Source/PixelStreaming2Core/Internal/PixelStreaming2PluginSettings.h b/Engine/Plugins/Media/PixelStreaming2/Source/PixelStreaming2Core/Internal/PixelStreaming2PluginSettings.h
--- a/Engine/Plugins/Media/PixelStreaming2/Source/PixelStreaming2Core/Internal/PixelStreaming2PluginSettings.h
+++ b/Engine/Plugins/Media/PixelStreaming2/Source/PixelStreaming2Core/Internal/PixelStreaming2PluginSettings.h
@@ -159,7 +159,7 @@ public:
 		DisplayName = "Capture Using Fence",
 		ToolTip = "Whether the texture copy we do during image capture should use a fence or not (non-fenced is faster but less safe)."
 		))
-	bool CaptureUseFence = false;
+	bool CaptureUseFence = true;
 
 	static TAutoConsoleVariable<bool> CVarDebugDumpAudio;
 
diff --git a/Engine/Plugins/Media/PixelStreaming2/Source/PixelStreaming2Core/Private/PixelStreaming2PluginSettings.cpp b/Engine/Plugins/Media/PixelStreaming2/Source/PixelStreaming2Core/Private/PixelStreaming2PluginSettings.cpp
--- a/Engine/Plugins/Media/PixelStreaming2/Source/PixelStreaming2Core/Private/PixelStreaming2PluginSettings.cpp
+++ b/Engine/Plugins/Media/PixelStreaming2/Source/PixelStreaming2Core/Private/PixelStreaming2PluginSettings.cpp
@@ -204,7 +204,7 @@ TAutoConsoleVariable<FString> UPixelStreaming2PluginSettings::CVarSignallingURL(
 
 TAutoConsoleVariable<bool> UPixelStreaming2PluginSettings::CVarCaptureUseFence(
 	TEXT("PixelStreaming2.CaptureUseFence"),
-	false,
+	true,
 	TEXT("Whether the texture copy we do during image capture should use a fence or not (non-fenced is faster but less safe)."),
 	FConsoleVariableDelegate::CreateLambda([](IConsoleVariable* Var) { Delegates()->OnCaptureUseFenceChanged.Broadcast(Var); }),
 	ECVF_Default);
