[PS2] Fix crash when closing Pixel Streaming and a stats task lasts longer then the GEngine. GEngine becomes null which crashes FStats when FStats calls a UE function that accesses GEngine.

Fix works by checking if the engine exit is requested before processing stats. There is still a small opportunity for GEngine to become null while stats are processing.

diff --git a/Engine/Plugins/Media/PixelStreaming2/Source/PixelStreaming2/Private/RTCStatsCollector.cpp b/Engine/Plugins/Media/PixelStreaming2/Source/PixelStreaming2/Private/RTCStatsCollector.cpp
--- a/Engine/Plugins/Media/PixelStreaming2/Source/PixelStreaming2/Private/RTCStatsCollector.cpp
+++ b/Engine/Plugins/Media/PixelStreaming2/Source/PixelStreaming2/Private/RTCStatsCollector.cpp
@@ -50,7 +50,7 @@ namespace UE::PixelStreaming2
 	void FRTCStatsCollector::Process(const EpicRtcConnectionStats& InStats)
 	{
 		FStats* PSStats = FStats::Get();
-		if (!bIsEnabled || !PSStats)
+		if (!bIsEnabled || !PSStats || IsEngineExitRequested())
 		{
 			return;
 		}
