[PS2] Fix: Simulcast highest layer not working

diff --git a/Engine/Plugins/Media/PixelStreaming2/Source/PixelStreaming2/Private/UtilsVideo.h b/Engine/Plugins/Media/PixelStreaming2/Source/PixelStreaming2/Private/UtilsVideo.h
--- a/Engine/Plugins/Media/PixelStreaming2/Source/PixelStreaming2/Private/UtilsVideo.h
+++ b/Engine/Plugins/Media/PixelStreaming2/Source/PixelStreaming2/Private/UtilsVideo.h
@@ -122,17 +122,20 @@ namespace UE::PixelStreaming2
 		if (UPixelStreaming2PluginSettings::CVarEncoderEnableSimulcast.GetValueOnAnyThread())
 		{
 			int MinBps = UPixelStreaming2PluginSettings::CVarWebRTCMinBitrate.GetValueOnAnyThread();
-			int MaxBps = UPixelStreaming2PluginSettings::CVarWebRTCMaxBitrate.GetValueOnAnyThread();
-			int BpsRange = MaxBps - MinBps;
+			// We have to halve the maximum here due to the layer assignments max bitrates effectively summing to 2
+			// 1/3 + 2/3 + 3/3
+			int MaxBps = UPixelStreaming2PluginSettings::CVarWebRTCMaxBitrate.GetValueOnAnyThread() / 2;
 
-			// Bitrates assignment per layer. Effectively
-			// 0: 0 -> 1/4
-			// 1: 1/4 -> 1/2
-			// 2: 1/2 -> 1
+			int OneThird = MaxBps / 3;
+			int TwoThird = 2 * MaxBps / 3;
+			// Bitrates assignment per layer
+			// 0: 0 -> 1/3
+			// 1: 1/3 -> 2/3
+			// 2: 2/3 -> 3/3
 			TArray<TTuple<int, int>> Bitrates = {
-				{ MinBps, MinBps + (BpsRange / 4) },
-				{ MinBps + (BpsRange / 4), MinBps + (BpsRange / 2) },
-				{ MinBps + (BpsRange / 2), MaxBps }
+				{ MinBps, OneThird },
+				{ OneThird, TwoThird },
+				{ TwoThird, MaxBps }
 			};
 
 			for (int i = 0; i < NumSimulcastLayers; i++)
diff --git a/Engine/Plugins/Media/PixelStreaming2/Source/PixelStreaming2/Private/VideoCapturer.cpp b/Engine/Plugins/Media/PixelStreaming2/Source/PixelStreaming2/Private/VideoCapturer.cpp
--- a/Engine/Plugins/Media/PixelStreaming2/Source/PixelStreaming2/Private/VideoCapturer.cpp
+++ b/Engine/Plugins/Media/PixelStreaming2/Source/PixelStreaming2/Private/VideoCapturer.cpp
@@ -151,7 +151,7 @@ namespace UE::PixelStreaming2
 			switch (FinalFormat)
 			{
 				case PixelCaptureBufferFormat::FORMAT_RHI:
-					if (VideoProducer->IsFrameAlreadyCopied() && FinalScale == 1.0 && UE::PixelStreaming2::GetSimulcastParameters().Num() == 1)
+					if (VideoProducer->IsFrameAlreadyCopied() && FinalScale == 1.0)
 					{
 						// If we only have a single layer (and it's scale is 1), we can use the no copy capturer
 						// as we know the output from the media capture will already be the correct format and scale
