[PS2] Fix: SoundGenerator being accessed post UObject destruction

diff --git a/Engine/Plugins/Media/PixelStreaming2/Source/PixelStreaming2/Private/Blueprints/PixelStreaming2AudioComponent.cpp b/Engine/Plugins/Media/PixelStreaming2/Source/PixelStreaming2/Private/Blueprints/PixelStreaming2AudioComponent.cpp
--- a/Engine/Plugins/Media/PixelStreaming2/Source/PixelStreaming2/Private/Blueprints/PixelStreaming2AudioComponent.cpp
+++ b/Engine/Plugins/Media/PixelStreaming2/Source/PixelStreaming2/Private/Blueprints/PixelStreaming2AudioComponent.cpp
@@ -47,6 +47,7 @@ void UPixelStreaming2AudioComponent::BeginDestroy()
 {
 	Super::BeginDestroy();
 	Reset();
+	SoundGenerator = nullptr;
 }
 
 bool UPixelStreaming2AudioComponent::ListenTo(FString PlayerToListenTo)
@@ -114,13 +115,16 @@ void UPixelStreaming2AudioComponent::Reset()
 {
 	PlayerToHear = FString();
 	StreamerToHear = FString();
-	SoundGenerator->bShouldGenerateAudio = false;
+	if (SoundGenerator)
+	{
+		SoundGenerator->bShouldGenerateAudio = false;
+		SoundGenerator->EmptyBuffers();
+	}
 	if (AudioSink)
 	{
 		AudioSink->RemoveAudioConsumer(this);
 	}
 	AudioSink = nullptr;
-	SoundGenerator->EmptyBuffers();
 }
 
 bool UPixelStreaming2AudioComponent::IsListeningToPlayer()
@@ -136,7 +140,7 @@ bool UPixelStreaming2AudioComponent::WillListenToAnyPlayer()
 void UPixelStreaming2AudioComponent::ConsumeRawPCM(const int16_t* AudioData, int InSampleRate, size_t NChannels, size_t NFrames)
 {
 	// Sound generator has not been initialized yet.
-	if (SoundGenerator->GetSampleRate() == 0 || GetAudioComponent() == nullptr)
+	if (!SoundGenerator || SoundGenerator->GetSampleRate() == 0 || GetAudioComponent() == nullptr)
 	{
 		return;
 	}
diff --git a/Engine/Plugins/Media/PixelStreaming2/Source/PixelStreaming2/Private/Blueprints/PixelStreaming2Peer.cpp b/Engine/Plugins/Media/PixelStreaming2/Source/PixelStreaming2/Private/Blueprints/PixelStreaming2Peer.cpp
--- a/Engine/Plugins/Media/PixelStreaming2/Source/PixelStreaming2/Private/Blueprints/PixelStreaming2Peer.cpp
+++ b/Engine/Plugins/Media/PixelStreaming2/Source/PixelStreaming2/Private/Blueprints/PixelStreaming2Peer.cpp
@@ -71,6 +71,13 @@ void UPixelStreaming2Peer::BeginPlay()
 	Super::BeginPlay();
 }
 
+void UPixelStreaming2Peer::BeginDestroy()
+{
+	Super::BeginDestroy();
+
+	SoundGenerator = nullptr;
+}
+
 void UPixelStreaming2Peer::EndPlay(const EEndPlayReason::Type EndPlayReason)
 {
 	Disconnect(FString(TEXT("UPixelStreaming2Peer::EndPlay called with reason: ")) + StaticEnum<EEndPlayReason::Type>()->GetNameStringByValue(EndPlayReason));
@@ -237,7 +244,7 @@ void UPixelStreaming2Peer::OnEndGenerate()
 void UPixelStreaming2Peer::ConsumeRawPCM(const int16_t* AudioData, int InSampleRate, size_t NChannels, size_t NFrames)
 {
 	// Sound generator has not been initialized yet.
-	if (SoundGenerator->GetSampleRate() == 0 || GetAudioComponent() == nullptr)
+	if (!SoundGenerator || SoundGenerator->GetSampleRate() == 0 || GetAudioComponent() == nullptr)
 	{
 		return;
 	}
@@ -268,8 +275,11 @@ void UPixelStreaming2Peer::OnConsumerAdded()
 
 void UPixelStreaming2Peer::OnConsumerRemoved()
 {
-	SoundGenerator->bShouldGenerateAudio = false;
-	SoundGenerator->EmptyBuffers();
+	if (SoundGenerator)
+	{
+		SoundGenerator->bShouldGenerateAudio = false;
+		SoundGenerator->EmptyBuffers();
+	}
 }
 
 void UPixelStreaming2Peer::OnSessionStateUpdate(const EpicRtcSessionState StateUpdate)
diff --git a/Engine/Plugins/Media/PixelStreaming2/Source/PixelStreaming2/Private/Blueprints/PixelStreaming2Peer.h b/Engine/Plugins/Media/PixelStreaming2/Source/PixelStreaming2/Private/Blueprints/PixelStreaming2Peer.h
--- a/Engine/Plugins/Media/PixelStreaming2/Source/PixelStreaming2/Private/Blueprints/PixelStreaming2Peer.h
+++ b/Engine/Plugins/Media/PixelStreaming2/Source/PixelStreaming2/Private/Blueprints/PixelStreaming2Peer.h
@@ -32,6 +32,7 @@ class UPixelStreaming2Peer : public USynthComponent, public IPixelStreaming2Audi
 protected:
 	// UObject overrides
 	virtual void BeginPlay() override;
+	virtual void BeginDestroy() override;
 	virtual void EndPlay(const EEndPlayReason::Type EndPlayReason) override;
 
 	// USynthComponent overrides
