[PS2] Fix: SendPixelStreaming2Response not working

diff --git a/Engine/Plugins/Media/PixelStreaming2/Source/PixelStreaming2/Private/Blueprints/PixelStreaming2InputComponent.cpp b/Engine/Plugins/Media/PixelStreaming2/Source/PixelStreaming2/Private/Blueprints/PixelStreaming2InputComponent.cpp
--- a/Engine/Plugins/Media/PixelStreaming2/Source/PixelStreaming2/Private/Blueprints/PixelStreaming2InputComponent.cpp
+++ b/Engine/Plugins/Media/PixelStreaming2/Source/PixelStreaming2/Private/Blueprints/PixelStreaming2InputComponent.cpp
@@ -58,6 +58,7 @@ void UPixelStreaming2Input::SendPixelStreaming2Response(const FString& Descripto
 	{
 		PixelStreaming2Module->ForEachStreamer([&Descriptor, this](TSharedPtr<IPixelStreaming2Streamer> Streamer) {
 			TSharedPtr<IPixelStreaming2InputHandler> Handler = Streamer->GetInputHandler().Pin();
+			if (!Handler)
 			{
 				UE_LOG(LogPixelStreaming2, Error, TEXT("Pixel Streaming input handler was null when sending response message."));
 				return;
