[PS2] Fix: StreamerComponent missing calls to Super

diff --git a/Engine/Plugins/Media/PixelStreaming2/Source/PixelStreaming2/Private/Blueprints/PixelStreaming2StreamerComponent.cpp b/Engine/Plugins/Media/PixelStreaming2/Source/PixelStreaming2/Private/Blueprints/PixelStreaming2StreamerComponent.cpp
--- a/Engine/Plugins/Media/PixelStreaming2/Source/PixelStreaming2/Private/Blueprints/PixelStreaming2StreamerComponent.cpp
+++ b/Engine/Plugins/Media/PixelStreaming2/Source/PixelStreaming2/Private/Blueprints/PixelStreaming2StreamerComponent.cpp
@@ -7,11 +7,13 @@
 #include "Slate/SceneViewport.h"
 
 UPixelStreaming2StreamerComponent::UPixelStreaming2StreamerComponent(const FObjectInitializer& ObjectInitializer)
+	: Super(ObjectInitializer)
 {
 }
 
 void UPixelStreaming2StreamerComponent::BeginPlay()
 {
+	Super::BeginPlay();
 	if (Streamer)
 	{
 		SetupStreamerInput();
@@ -20,6 +22,7 @@ void UPixelStreaming2StreamerComponent::BeginPlay()
 
 void UPixelStreaming2StreamerComponent::EndPlay(const EEndPlayReason::Type EndPlayReason)
 {
+	Super::EndPlay(EndPlayReason);
 	if (Streamer)
 	{
 		IPixelStreaming2Module::Get().DeleteStreamer(Streamer);
