[PS2] Fix: Multiple gamepad support

diff --git a/Engine/Plugins/Media/PixelStreaming2/Source/PixelStreaming2Input/Private/InputDevice.cpp b/Engine/Plugins/Media/PixelStreaming2/Source/PixelStreaming2Input/Private/InputDevice.cpp
--- a/Engine/Plugins/Media/PixelStreaming2/Source/PixelStreaming2Input/Private/InputDevice.cpp
+++ b/Engine/Plugins/Media/PixelStreaming2/Source/PixelStreaming2Input/Private/InputDevice.cpp
@@ -2,6 +2,7 @@
 
 #include "InputDevice.h"
 #include "Framework/Application/SlateApplication.h"
+#include "GenericPlatform/GenericPlatformInputDeviceMapper.h"
 
 TSharedPtr<UE::PixelStreaming2Input::FInputDevice> UE::PixelStreaming2Input::FInputDevice::InputDevice;
 
@@ -14,11 +15,7 @@ namespace UE::PixelStreaming2Input
 			return InputDevice;
 		}
 
-		TSharedPtr<FInputDevice> Device = TSharedPtr<FInputDevice>(new FInputDevice());
-		if (Device)
-		{
-			InputDevice = Device;
-		}
+		InputDevice = TSharedPtr<FInputDevice>(new FInputDevice());
 		return InputDevice;
 	}
 
@@ -92,18 +89,52 @@ namespace UE::PixelStreaming2Input
 
 	uint8 FInputDevice::OnControllerConnected()
 	{
+		TArray<uint8> InputDeviceKeys;
+		InputDevices.GenerateKeyArray(InputDeviceKeys);
+
+		// Find the next available controller id. eg [0, 2, 3] will return 1
 		uint8 NextControllerId = 0;
-		while (ConnectedControllers.Contains(NextControllerId))
+		while (InputDeviceKeys.Contains(NextControllerId))
 		{
 			NextControllerId++;
 		}
 
-		ConnectedControllers.Add(NextControllerId);
+		IPlatformInputDeviceMapper& DeviceMapper = IPlatformInputDeviceMapper::Get();
+		FInputDeviceId				DeviceId = DeviceMapper.AllocateNewInputDeviceId();
+		FPlatformUserId				UserId = DeviceMapper.GetPlatformUserForNewlyConnectedDevice();
+		DeviceMapper.Internal_MapInputDeviceToUser(DeviceId, UserId, EInputDeviceConnectionState::Connected);
+
+		InputDevices.Add(NextControllerId, { DeviceId, UserId });
+
 		return NextControllerId;
 	}
 
 	void FInputDevice::OnControllerDisconnected(uint8 DeleteControllerId)
 	{
-		ConnectedControllers.Remove(DeleteControllerId);
+		FInputDeviceId							DeviceId = INPUTDEVICEID_NONE;
+		FPlatformUserId							PlatformUserId = PLATFORMUSERID_NONE;
+		if (!GetPlatformUserAndDevice(DeleteControllerId, DeviceId, PlatformUserId))
+		{
+			return;
+		}
+		IPlatformInputDeviceMapper& DeviceMapper = IPlatformInputDeviceMapper::Get();
+		DeviceMapper.Internal_MapInputDeviceToUser(DeviceId, DeviceMapper.GetUserForUnpairedInputDevices(), EInputDeviceConnectionState::Disconnected);
+
+		InputDevices.Remove(DeleteControllerId);
+	}
+
+	bool FInputDevice::GetPlatformUserAndDevice(uint8 ControllerId, FInputDeviceId& OutDeviceId, FPlatformUserId& OutPlatformUser)
+	{
+		if (!InputDevices.Contains(ControllerId))
+		{
+			OutDeviceId = INPUTDEVICEID_NONE;
+			OutPlatformUser = PLATFORMUSERID_NONE;
+			return false;
+		}
+
+		TPair<FInputDeviceId, FPlatformUserId> DevicePair = InputDevices[ControllerId];
+		OutDeviceId = DevicePair.Get<0>();
+		OutPlatformUser = DevicePair.Get<1>();
+		return true;
 	}
 } // namespace UE::PixelStreaming2Input
diff --git a/Engine/Plugins/Media/PixelStreaming2/Source/PixelStreaming2Input/Private/InputDevice.h b/Engine/Plugins/Media/PixelStreaming2/Source/PixelStreaming2Input/Private/InputDevice.h
--- a/Engine/Plugins/Media/PixelStreaming2/Source/PixelStreaming2Input/Private/InputDevice.h
+++ b/Engine/Plugins/Media/PixelStreaming2/Source/PixelStreaming2Input/Private/InputDevice.h
@@ -4,6 +4,7 @@
 
 #include "IInputDevice.h"
 #include "IPixelStreaming2InputHandler.h"
+#include "Misc/CoreMiscDefines.h"
 
 namespace UE::PixelStreaming2Input
 {
@@ -29,6 +30,7 @@ namespace UE::PixelStreaming2Input
 		void  AddInputHandler(TSharedPtr<IPixelStreaming2InputHandler> InputHandler);
 		uint8 OnControllerConnected();
 		void  OnControllerDisconnected(uint8 DeleteControllerId);
+		bool  GetPlatformUserAndDevice(uint8 ControllerId, FInputDeviceId& OutDeviceId, FPlatformUserId& OutPlatformUser);
 
 		static TSharedPtr<FInputDevice> GetInputDevice();
 
@@ -51,12 +53,10 @@ namespace UE::PixelStreaming2Input
 		TArray<TWeakPtr<IPixelStreaming2InputHandler>> InputHandlers;
 
 		/**
-		 * The array of connected controllers. As each device can have multiple controllers, we want to make sure that each controller of each device is unique
+		 * The map of connected controllers. As each handler can have multiple input devices, we want to make sure that each controller of each device is unique
 		 * As such, a simple incrementer approach is not applicable and we must instead keep track of all the connected controllers
 		 *
-		 * NOTE: We store each input handler as a weakptr as we don't want to be the reason
-		 * the handler doesn't get deleted. Each handler should be tied to a Streamer's lifecycle
 		 */
-		TArray<uint8> ConnectedControllers;
+		TMap<uint8, TPair<FInputDeviceId, FPlatformUserId>> InputDevices;
 	};
 } // namespace UE::PixelStreaming2Input
\ No newline at end of file
diff --git a/Engine/Plugins/Media/PixelStreaming2/Source/PixelStreaming2Input/Private/InputHandler.cpp b/Engine/Plugins/Media/PixelStreaming2/Source/PixelStreaming2Input/Private/InputHandler.cpp
--- a/Engine/Plugins/Media/PixelStreaming2/Source/PixelStreaming2Input/Private/InputHandler.cpp
+++ b/Engine/Plugins/Media/PixelStreaming2/Source/PixelStreaming2Input/Private/InputHandler.cpp
@@ -37,7 +37,7 @@
 namespace UE::PixelStreaming2Input
 {
 	typedef EPixelStreaming2InputAction Action;
-	typedef EPixelStreaming2XRSystem	  XRSystem;
+	typedef EPixelStreaming2XRSystem	XRSystem;
 
 	FPixelStreaming2InputHandler::FPixelStreaming2InputHandler(TSharedPtr<FPixelStreaming2ApplicationWrapper> InApplicationWrapper, const TSharedPtr<FGenericApplicationMessageHandler>& InTargetHandler)
 		: TargetViewport(nullptr)
@@ -588,8 +588,7 @@ namespace UE::PixelStreaming2Input
 
 	void FPixelStreaming2InputHandler::HandleOnControllerConnected(FMemoryReader Ar)
 	{
-		TSharedPtr<FInputDevice> InputDevice = FInputDevice::GetInputDevice();
-		uint8					 NextControllerId = InputDevice->OnControllerConnected();
+		uint8 NextControllerId = FInputDevice::GetInputDevice()->OnControllerConnected();
 
 		// When a virtual controller (from the browser) is "connected" into UE's input system,
 		// it creates and id. That id is used to differentitate each controller used.
@@ -610,10 +609,15 @@ namespace UE::PixelStreaming2Input
 	void FPixelStreaming2InputHandler::HandleOnControllerAnalog(FMemoryReader Ar)
 	{
 		const TPayloadThreeParam<uint8, uint8, double> Payload(Ar);
-		const FInputDeviceId						   ControllerId = FInputDeviceId::CreateFromInternalId((int32)Payload.Param1);
 
-		FKeyId KeyId = Payload.Param2;
+		FInputDeviceId	DeviceId = INPUTDEVICEID_NONE;
+		FPlatformUserId PlatformUserId = PLATFORMUSERID_NONE;
+		if (!FInputDevice::GetInputDevice()->GetPlatformUserAndDevice(Payload.Param1, DeviceId, PlatformUserId))
+		{
+			return;
+		}
 		double AxisValue = Payload.Param3;
+		FKeyId KeyId = Payload.Param2;
 
 		FKey* AnalogKeyPtr = FInputConverter::GamepadInputToFKey.Find(MakeTuple(KeyId, Action::Axis));
 		if (!AnalogKeyPtr)
@@ -627,74 +631,56 @@ namespace UE::PixelStreaming2Input
 		AnalogValue.bKeepUnlessZero = (KeyId == 5 || KeyId == 6);
 
 		// Overwrite the last data: every tick only process the latest
-		AnalogEventsReceivedThisTick.FindOrAdd(ControllerId).FindOrAdd(AnalogKeyPtr) = AnalogValue;
+		AnalogEventsReceivedThisTick.FindOrAdd(DeviceId).FindOrAdd(AnalogKeyPtr) = AnalogValue;
 	}
 
 	void FPixelStreaming2InputHandler::HandleOnControllerButtonPressed(FMemoryReader Ar)
 	{
 		TPayloadThreeParam<uint8, uint8, uint8> Payload(Ar);
-		FKey*									ButtonPtr = FInputConverter::GamepadInputToFKey.Find(MakeTuple(Payload.Param2, Action::Click));
+		FInputDeviceId							DeviceId = INPUTDEVICEID_NONE;
+		FPlatformUserId							PlatformUserId = PLATFORMUSERID_NONE;
+		if (!FInputDevice::GetInputDevice()->GetPlatformUserAndDevice(Payload.Param1, DeviceId, PlatformUserId))
+		{
+			return;
+		}
+		bool  bIsRepeat = Payload.Param3 != 0;
+		FKey* ButtonPtr = FInputConverter::GamepadInputToFKey.Find(MakeTuple(Payload.Param2, Action::Click));
 		if (ButtonPtr == nullptr)
 		{
 			return;
 		}
 
-		FInputDeviceId	ControllerId = FInputDeviceId::CreateFromInternalId((int32)Payload.Param1);
-		bool			bIsRepeat = Payload.Param3 != 0;
-		FPlatformUserId UserId = IPlatformInputDeviceMapper::Get().GetPrimaryPlatformUser();
-
-		FKeyEvent KeyEvent(
-			*ButtonPtr,															  /* InKey */
-			FSlateApplication::Get().GetPlatformApplication()->GetModifierKeys(), /* InModifierKeys */
-			ControllerId,														  /* InDeviceId */
-			bIsRepeat,															  /* bInIsRepeat */
-			0,																	  /* InCharacterCode */
-			0,																	  /* InKeyCode */
-			// TODO (william.belcher): This user idx should be the playerId
-			0 /* InUserIndex*/
-		);
-		FSlateApplication::Get().ProcessKeyDownEvent(KeyEvent);
+		MessageHandler->OnControllerButtonPressed(ButtonPtr->GetFName(), PlatformUserId, DeviceId, bIsRepeat);
 
-		UE_LOG(LogPixelStreaming2Input, Verbose, TEXT("GAMEPAD_PRESSED: ControllerId = %d; KeyName = %s; IsRepeat = %s;"), ControllerId.GetId(), *ButtonPtr->ToString(), bIsRepeat ? TEXT("True") : TEXT("False"));
+		UE_LOG(LogPixelStreaming2Input, Verbose, TEXT("GAMEPAD_PRESSED: ControllerId = %d; KeyName = %s; IsRepeat = %s;"), Payload.Param1, *ButtonPtr->ToString(), bIsRepeat ? TEXT("True") : TEXT("False"));
 	}
 
 	void FPixelStreaming2InputHandler::HandleOnControllerButtonReleased(FMemoryReader Ar)
 	{
 		TPayloadTwoParam<uint8, uint8> Payload(Ar);
-		FKey*						   ButtonPtr = FInputConverter::GamepadInputToFKey.Find(MakeTuple(Payload.Param2, Action::Click));
+		FInputDeviceId				   DeviceId = INPUTDEVICEID_NONE;
+		FPlatformUserId				   PlatformUserId = PLATFORMUSERID_NONE;
+		if (!FInputDevice::GetInputDevice()->GetPlatformUserAndDevice(Payload.Param1, DeviceId, PlatformUserId))
+		{
+			return;
+		}
+		FKey* ButtonPtr = FInputConverter::GamepadInputToFKey.Find(MakeTuple(Payload.Param2, Action::Click));
 		if (ButtonPtr == nullptr)
 		{
 			return;
 		}
 
-		FInputDeviceId	ControllerId = FInputDeviceId::CreateFromInternalId((int32)Payload.Param1);
-		FPlatformUserId UserId = IPlatformInputDeviceMapper::Get().GetPrimaryPlatformUser();
+		MessageHandler->OnControllerButtonReleased(ButtonPtr->GetFName(), PlatformUserId, DeviceId, false);
 
-		FKeyEvent KeyEvent(
-			*ButtonPtr,															  /* InKey */
-			FSlateApplication::Get().GetPlatformApplication()->GetModifierKeys(), /* InModifierKeys */
-			ControllerId,														  /* InDeviceId */
-			false,																  /* bInIsRepeat */
-			0,																	  /* InCharacterCode */
-			0,																	  /* InKeyCode */
-			// TODO (william.belcher): This user idx should be the playerId
-			0 /* InUserIndex*/
-		);
-		FSlateApplication::Get().ProcessKeyUpEvent(KeyEvent);
-
-		UE_LOG(LogPixelStreaming2Input, Verbose, TEXT("GAMEPAD_RELEASED: ControllerId = %d; KeyName = %s;"), ControllerId.GetId(), *ButtonPtr->ToString());
+		UE_LOG(LogPixelStreaming2Input, Verbose, TEXT("GAMEPAD_RELEASED: ControllerId = %d; KeyName = %s;"), Payload.Param1, *ButtonPtr->ToString());
 	}
 
 	void FPixelStreaming2InputHandler::HandleOnControllerDisconnected(FMemoryReader Ar)
 	{
 		TPayloadOneParam<uint8> Payload(Ar);
+		FInputDevice::GetInputDevice()->OnControllerDisconnected(Payload.Param1);
 
-		uint8 DisconnectedControllerId = Payload.Param1;
-
-		TSharedPtr<FInputDevice> InputDevice = FInputDevice::GetInputDevice();
-		InputDevice->OnControllerDisconnected(DisconnectedControllerId);
-
-		UE_LOG(LogPixelStreaming2Input, Verbose, TEXT("GAMEPAD_DISCONNECTED: ControllerId = %d"), DisconnectedControllerId);
+		UE_LOG(LogPixelStreaming2Input, Verbose, TEXT("GAMEPAD_DISCONNECTED: ControllerId = %d"), Payload.Param1);
 	}
 
 	/**
@@ -1042,7 +1028,7 @@ namespace UE::PixelStreaming2Input
 			return;
 		}
 
-		MessageHandler->OnControllerButtonPressed(ButtonPtr->GetFName(), DeviceMapper.GetPrimaryPlatformUser(), ControllerId, bIsRepeat);
+		MessageHandler->OnControllerButtonPressed(ButtonPtr->GetFName(), PLATFORMUSERID_NONE /* Not used */, ControllerId, bIsRepeat);
 		UE_LOG(LogPixelStreaming2Input, Verbose, TEXT("XR_TOUCHED: ControllerId = %d; KeyName = %s; IsRepeat = %s;"), ControllerId.GetId(), *ButtonPtr->ToString(), bIsRepeat ? TEXT("True") : TEXT("False"));
 	}
 
@@ -1062,7 +1048,7 @@ namespace UE::PixelStreaming2Input
 			return;
 		}
 
-		MessageHandler->OnControllerButtonReleased(ButtonPtr->GetFName(), DeviceMapper.GetPrimaryPlatformUser(), ControllerId, bIsRepeat);
+		MessageHandler->OnControllerButtonReleased(ButtonPtr->GetFName(), PLATFORMUSERID_NONE /* Not used */, ControllerId, bIsRepeat);
 		UE_LOG(LogPixelStreaming2Input, Verbose, TEXT("XR_TOUCH_RELEASED: ControllerId = %d; KeyName = %s; IsRepeat = %s;"), ControllerId.GetId(), *ButtonPtr->ToString(), bIsRepeat ? TEXT("True") : TEXT("False"));
 	}
 
@@ -1085,7 +1071,7 @@ namespace UE::PixelStreaming2Input
 		// Only send button click if there is no axis associated with this button or this the first press of the axis
 		if (ButtonPtr)
 		{
-			MessageHandler->OnControllerButtonPressed(ButtonPtr->GetFName(), DeviceMapper.GetPrimaryPlatformUser(), ControllerId, bIsRepeat);
+			MessageHandler->OnControllerButtonPressed(ButtonPtr->GetFName(), PLATFORMUSERID_NONE /* Not used */, ControllerId, bIsRepeat);
 			UE_LOG(LogPixelStreaming2Input, Verbose, TEXT("XR_PRESSED: ControllerId = %d; KeyName = %s; IsRepeat = %s"), ControllerId.GetId(), *ButtonPtr->ToString(), bIsRepeat ? TEXT("True") : TEXT("False"));
 		}
 
@@ -1135,7 +1121,7 @@ namespace UE::PixelStreaming2Input
 		FKey* ButtonPtr = FInputConverter::XRInputToFKey.Find(MakeTuple(System, Handedness, ButtonIdx, Action::Click));
 		if (ButtonPtr)
 		{
-			MessageHandler->OnControllerButtonReleased(ButtonPtr->GetFName(), DeviceMapper.GetPrimaryPlatformUser(), ControllerId, bIsRepeat);
+			MessageHandler->OnControllerButtonReleased(ButtonPtr->GetFName(), PLATFORMUSERID_NONE /* Not used */, ControllerId, bIsRepeat);
 			UE_LOG(LogPixelStreaming2Input, Verbose, TEXT("XR_RELEASED: ControllerId = %d; KeyName = %s; IsRepeat = %s;"), ControllerId.GetId(), *ButtonPtr->ToString(), bIsRepeat ? TEXT("True") : TEXT("False"));
 		}
 	}
@@ -1162,7 +1148,7 @@ namespace UE::PixelStreaming2Input
 		AnalogValue.Value = Payload.Param3;
 		AnalogValue.bIsRepeat = false;
 
-		UE_LOG(LogPixelStreaming2Input, Verbose, TEXT("XR_ANALOG: ControllerId = %d; KeyName = %s; IsRepeat = False; AnalogValue = %.4f; [Queued for Tick()]"), ControllerId.GetId(), *AnalogKeyPtr->ToString(), AnalogValue.Value);
+		UE_LOG(LogPixelStreaming2Input, VeryVerbose, TEXT("XR_ANALOG: ControllerId = %d; KeyName = %s; IsRepeat = False; AnalogValue = %.4f; [Queued for Tick()]"), ControllerId.GetId(), *AnalogKeyPtr->ToString(), AnalogValue.Value);
 		AnalogEventsReceivedThisTick.FindOrAdd(ControllerId).FindOrAdd(AnalogKeyPtr) = AnalogValue;
 	}
 
@@ -1395,28 +1381,14 @@ namespace UE::PixelStreaming2Input
 				FKey*				  Key = FKeyIt->Key;
 				FAnalogValue		  AnalogValue = FKeyIt->Value;
 				bool				  bIsRepeat = AnalogValue.bIsRepeat;
-
 				if (!Key)
 				{
 					return;
 				}
 
 				// Pass an analog input along the engine's input processing system
-				FSlateApplication&		SlateApplication = FSlateApplication::Get();
-				const FAnalogInputEvent AnalogInputEvent(
-					*Key,														  /* InKey */
-					SlateApplication.GetPlatformApplication()->GetModifierKeys(), /* InModifierKeys */
-					ControllerId,												  /* InDeviceId */
-					bIsRepeat,													  /* bInIsRepeat*/
-					0,															  /* InCharacterCode */
-					0,															  /* InKeyCode */
-					AnalogValue.Value,											  /* InAnalogValue */
-					// TODO (william.belcher): This user idx should be the playerId
-					0 /* InUserIndex */
-				);
-
-				bool bHandled = FSlateApplication::Get().ProcessAnalogInputEvent(AnalogInputEvent);
-				UE_LOG(LogPixelStreaming2Input, Verbose, TEXT("TICKED ANALOG Input: ControllerId = %d; KeyName = %s; IsRepeat = %s; AnalogValue = %.4f; Handled = %s; [Queued for Tick()]"), ControllerId.GetId(), *Key->ToString(), bIsRepeat ? TEXT("True") : TEXT("False"), AnalogValue.Value, bHandled ? TEXT("True") : TEXT("False"));
+				bool bHandled = MessageHandler->OnControllerAnalog(Key->GetFName(), PLATFORMUSERID_NONE /* Not used */, ControllerId, AnalogValue.Value);
+				UE_LOG(LogPixelStreaming2Input, VeryVerbose, TEXT("TICKED ANALOG Input: ControllerId = %d; KeyName = %s; IsRepeat = %s; AnalogValue = %.4f; Handled = %s; [Queued for Tick()]"), ControllerId.GetId(), *Key->ToString(), bIsRepeat ? TEXT("True") : TEXT("False"), AnalogValue.Value, bHandled ? TEXT("True") : TEXT("False"));
 
 				// Remove current analog key unless it has the special `bKeepUnlessZero` flag set.
 				// This flag is used to continuously apply input values across ticks because
