[PS2] PR #12718: Refactor PixelStreaming2 plugin to enable other streaming protocols

diff --git a/Engine/Plugins/Media/PixelStreaming2/Config/BasePixelStreaming2.ini b/Engine/Plugins/Media/PixelStreaming2/Config/BasePixelStreaming2.ini
--- a/Engine/Plugins/Media/PixelStreaming2/Config/BasePixelStreaming2.ini
+++ b/Engine/Plugins/Media/PixelStreaming2/Config/BasePixelStreaming2.ini
@@ -1,4 +1,15 @@
-[/Script/PixelStreaming2Core.PixelStreaming2PluginSettings]
+[/Script/PixelStreaming2Settings.PixelStreaming2PluginSettings]
 DefaultCursorClassName=/PixelStreaming2/DefaultCursor.DefaultCursor_C
 TextEditBeamCursorClassName=/PixelStreaming2/TextEditBeamCursor.TextEditBeamCursor_C
-HiddenCursorClassName=/PixelStreaming2/HiddenCursor.HiddenCursor_C
\ No newline at end of file
+HiddenCursorClassName=/PixelStreaming2/HiddenCursor.HiddenCursor_C
+
+[CoreRedirects]
++ClassRedirects=(OldName="/Script/PixelStreaming2.PixelStreaming2VideoProducer",NewName="/Script/PixelStreaming2.PixelStreaming2VideoProducerBase")
+# Remap all the old PixelStreaming2 blueprint classes to PixelStreaming2RTC
++ClassRedirects=(OldName="/Script/PixelStreaming2.PixelStreaming2Peer",NewName="/Script/PixelStreaming2RTC.PixelStreaming2Peer")
+
+# Remap old PixelStreaming2Core to PixelStreaming2Settings
++ClassRedirects=(OldName="/Script/PixelStreaming2Core.PixelStreaming2PluginSettings",NewName="/Script/PixelStreaming2Settings.PixelStreaming2PluginSettings")
++EnumRedirects=(OldName="/Script/PixelStreaming2Core.EPortAllocatorFlags",NewName="/Script/PixelStreaming2Settings.EPortAllocatorFlags")
++EnumRedirects=(OldName="/Script/PixelStreaming2Core.EInputControllerMode",NewName="/Script/PixelStreaming2Settings.EInputControllerMode")
++EnumRedirects=(OldName="/Script/PixelStreaming2Core.EPixelStreaming2EditorStreamTypes",NewName="/Script/PixelStreaming2Settings.EPixelStreaming2EditorStreamTypes")
diff --git a/Engine/Plugins/Media/PixelStreaming2/PixelStreaming2.uplugin b/Engine/Plugins/Media/PixelStreaming2/PixelStreaming2.uplugin
--- a/Engine/Plugins/Media/PixelStreaming2/PixelStreaming2.uplugin
+++ b/Engine/Plugins/Media/PixelStreaming2/PixelStreaming2.uplugin
@@ -16,14 +16,21 @@
 	"Installed": false,
 	"Modules": [
 		{
-			"Name": "PixelStreaming2Core",
+			"Name": "PixelStreaming2",
 			"Type": "Runtime",
 			"LoadingPhase": "PreDefault",
 			"PlatformAllowList": [ "Win64", "Linux", "Mac" ],
 			"TargetDenyList": [ ]
 		},
 		{
-			"Name": "PixelStreaming2",
+			"Name": "PixelStreaming2Settings",
+			"Type": "Runtime",
+			"LoadingPhase": "PreDefault",
+			"PlatformAllowList": [ "Win64", "Linux", "Mac" ],
+			"TargetDenyList": [ ]
+		},
+		{
+			"Name": "PixelStreaming2RTC",
 			"Type": "Runtime",
 			"LoadingPhase": "PreDefault",
 			"PlatformAllowList": [ "Win64", "Linux", "Mac" ],
diff --git a/Engine/Plugins/Media/PixelStreaming2/Source/PixelStreaming2/Internal/AudioCapturer.h b/Engine/Plugins/Media/PixelStreaming2/Source/PixelStreaming2/Internal/AudioCapturer.h
new file mode 100644
--- /dev/null
+++ b/Engine/Plugins/Media/PixelStreaming2/Source/PixelStreaming2/Internal/AudioCapturer.h
@@ -0,0 +1,110 @@
+// Copyright Epic Games, Inc. All Rights Reserved.
+
+#pragma once
+
+#include "DSP/MultithreadedPatching.h"
+#include "AudioProducer.h"
+#include "ISubmixBufferListener.h"
+#include "Logging.h"
+#include "SampleBuffer.h"
+#include "TickableTask.h"
+
+namespace UE::PixelStreaming2
+{
+	class FAudioCapturer;
+
+	class PIXELSTREAMING2_API FAudioPatchMixer : public Audio::FPatchMixer
+	{
+	public:
+		FAudioPatchMixer(uint8 NumChannels, uint32 SampleRate, float SampleSizeSeconds);
+		virtual ~FAudioPatchMixer() = default;
+
+		uint32 GetMaxBufferSize() const;
+		uint8  GetNumChannels() const;
+		uint32 GetSampleRate() const;
+
+	protected:
+		uint8  NumChannels;
+		uint32 SampleRate;
+		float  SampleSizeSeconds;
+	};
+
+	class PIXELSTREAMING2_API FPatchInputProxy : public IPixelStreaming2AudioProducer
+	{
+	public:
+		FPatchInputProxy(TSharedPtr<FAudioPatchMixer> InMixer);
+		virtual ~FPatchInputProxy() override;
+
+		virtual void PushAudio(const float* AudioData, int32 InNumSamples, int32 InNumChannels, int32 InSampleRate) override;
+
+	protected:
+		TSharedPtr<FAudioPatchMixer> Mixer;
+		Audio::FResampler			 Resampler;
+		Audio::FPatchInput			 PatchInput;
+		uint8						 NumChannels;
+		uint32						 SampleRate;
+	};
+
+	class PIXELSTREAMING2_API FMixAudioTask : public FPixelStreamingTickableTask
+	{
+	public:
+		FMixAudioTask(FAudioCapturer* Capturer, TSharedPtr<FAudioPatchMixer> Mixer);
+
+		virtual ~FMixAudioTask() = default;
+
+		// Begin FPixelStreamingTickableTask
+		virtual void		   Tick(float DeltaMs) override;
+		virtual const FString& GetName() const override;
+		// End FPixelStreamingTickableTask
+
+	protected:
+		bool								  bIsRunning;
+		Audio::VectorOps::FAlignedFloatBuffer MixingBuffer;
+
+		FAudioCapturer*				 Capturer;
+		TSharedPtr<FAudioPatchMixer> Mixer;
+	};
+
+	class PIXELSTREAMING2_API FAudioCapturer : public IPixelStreaming2AudioProducer
+	{
+	public:
+		static TSharedPtr<FAudioCapturer> Create(const int InSampleRate = 48000, const int InNumChannels = 2, const float InSampleSizeInSeconds = 0.5f);
+		virtual ~FAudioCapturer() = default;
+
+		// Mixed audio input will push its audio to an FPatchInputProxy for mixing
+		TSharedPtr<FAudioProducer> CreateAudioProducer();
+		void					   CreateAudioProducer(Audio::FDeviceId AudioDeviceId);
+		void					   RemoveAudioProducer(Audio::FDeviceId AudioDeviceId);
+
+		virtual void PushAudio(const float* AudioData, int32 NumSamples, int32 NumChannels, int32 SampleRate) override;
+
+		/**
+		 * This is broadcast each time audio is captured. Tracks should bind to this and push the audio into the track
+		 */
+		DECLARE_EVENT_FourParams(FAudioCapturer, FOnAudioBuffer, const int16_t*, int32, int32, const int32);
+		FOnAudioBuffer OnAudioBuffer;
+
+	protected:
+		FAudioCapturer(const int SampleRate = 48000, const int NumChannels = 2, const float SampleSizeInSeconds = 0.5f);
+
+		void OnDebugDumpAudioChanged(IConsoleVariable* Var);
+		void OnEnginePreExit();
+		void WriteDebugAudio();
+		// Allow the mix audio task to call OnAudio.
+		friend class FMixAudioTask;
+		void OnAudio(const float* AudioData, int32 NumSamples, int32 NumChannels, int32 SampleRate);
+
+	protected:
+		TSharedPtr<FAudioPatchMixer>  Mixer;
+		TUniqueTaskPtr<FMixAudioTask> MixerTask;
+
+		TMap<Audio::FDeviceId, TSharedPtr<FAudioProducer>> AudioProducers;
+
+		int				  SampleRate;
+		int				  NumChannels;
+		float			  SampleSizeSeconds;
+		Audio::FResampler Resampler;
+
+		Audio::TSampleBuffer<int16_t> DebugDumpAudioBuffer;
+	};
+} // namespace UE::PixelStreaming2
\ No newline at end of file
diff --git a/Engine/Plugins/Media/PixelStreaming2/Source/PixelStreaming2/Internal/AudioProducer.h b/Engine/Plugins/Media/PixelStreaming2/Source/PixelStreaming2/Internal/AudioProducer.h
new file mode 100644
--- /dev/null
+++ b/Engine/Plugins/Media/PixelStreaming2/Source/PixelStreaming2/Internal/AudioProducer.h
@@ -0,0 +1,39 @@
+// Copyright Epic Games, Inc. All Rights Reserved.
+
+#pragma once
+
+#include "Audio.h"
+#include "ISubmixBufferListener.h"
+#include "IPixelStreaming2AudioProducer.h"
+
+namespace UE::PixelStreaming2
+{
+	class FPatchInputProxy;
+
+	/**
+	 * An audio input capable of listening to UE submix's as well as receiving user audio via the PushAudio method.
+	 * Any received audio will be passed into the Parent's PushAudio method
+	 */
+	class PIXELSTREAMING2_API FAudioProducer : public ISubmixBufferListener, public IPixelStreaming2AudioProducer
+	{
+	public:
+		static TSharedPtr<FAudioProducer> Create(Audio::FDeviceId AudioDeviceId, TSharedPtr<FPatchInputProxy> PatchInput);
+		static TSharedPtr<FAudioProducer> Create(TSharedPtr<FPatchInputProxy> PatchInput);
+		virtual ~FAudioProducer() = default;
+
+		// For users to manually push non-submix audio into
+		virtual void PushAudio(const float* AudioData, int32 NumSamples, int32 NumChannels, int32 SampleRate) override;
+
+		// ISubmixBufferListener interface
+		virtual void OnNewSubmixBuffer(const USoundSubmix* OwningSubmix, float* AudioData, int32 NumSamples, int32 NumChannels, const int32 SampleRate, double AudioClock) override;
+
+		void ToggleMuted() { bIsMuted = !bIsMuted; }
+
+	protected:
+		FAudioProducer(TSharedPtr<FPatchInputProxy> PatchInput);
+
+	protected:
+		TSharedPtr<FPatchInputProxy> PatchInput;
+		bool						 bIsMuted = false;
+	};
+} // namespace UE::PixelStreaming2
\ No newline at end of file
diff --git a/Engine/Plugins/Media/PixelStreaming2/Source/PixelStreaming2/Internal/AudioSource.h b/Engine/Plugins/Media/PixelStreaming2/Source/PixelStreaming2/Internal/AudioSource.h
new file mode 100644
--- /dev/null
+++ b/Engine/Plugins/Media/PixelStreaming2/Source/PixelStreaming2/Internal/AudioSource.h
@@ -0,0 +1,21 @@
+// Copyright Epic Games, Inc. All Rights Reserved.
+
+#pragma once
+
+#include "Templates/SharedPointer.h"
+
+namespace UE::PixelStreaming2
+{
+	class PIXELSTREAMING2_API FAudioSource
+	{
+	public:
+		virtual ~FAudioSource() = default;
+
+		virtual void OnAudioBuffer(const int16_t* AudioData, int32 NumSamples, int32 NumChannels, const int32 SampleRate);
+
+		void SetMuted(bool bIsMuted);
+
+	protected:
+		bool bIsMuted = false;
+	};
+} // namespace UE::PixelStreaming2
\ No newline at end of file
diff --git a/Engine/Plugins/Media/PixelStreaming2/Source/PixelStreaming2/Internal/Blueprints/PixelStreaming2VideoProducer.h b/Engine/Plugins/Media/PixelStreaming2/Source/PixelStreaming2/Internal/Blueprints/PixelStreaming2VideoProducer.h
--- a/Engine/Plugins/Media/PixelStreaming2/Source/PixelStreaming2/Internal/Blueprints/PixelStreaming2VideoProducer.h
+++ b/Engine/Plugins/Media/PixelStreaming2/Source/PixelStreaming2/Internal/Blueprints/PixelStreaming2VideoProducer.h
@@ -11,7 +11,7 @@
  * The base video producer blueprint class
  */
 UCLASS(NotBlueprintType, NotBlueprintable, Category = "PixelStreaming2", META = (DisplayName = "Video Producer"))
-class PIXELSTREAMING2_API UPixelStreaming2VideoProducer : public UObject
+class PIXELSTREAMING2_API UPixelStreaming2VideoProducerBase : public UObject
 {
 	GENERATED_UCLASS_BODY()
 
@@ -26,7 +26,7 @@ protected:
  * A video producer that streams the Unreal Engine's back buffer
  */
 UCLASS(NotBlueprintType, NotBlueprintable, Category = "PixelStreaming2", META = (DisplayName = "Back Buffer Video Producer"))
-class PIXELSTREAMING2_API UPixelStreaming2VideoProducerBackBuffer : public UPixelStreaming2VideoProducer
+class PIXELSTREAMING2_API UPixelStreaming2VideoProducerBackBuffer : public UPixelStreaming2VideoProducerBase
 {
 	GENERATED_UCLASS_BODY()
 
@@ -38,7 +38,7 @@ public:
  * A video producer that streams the Unreal Engine's back buffer but uses the MediaIO Framework to capture the frame rather than Pixel Capture
  */
 UCLASS(NotBlueprintType, NotBlueprintable, Category = "PixelStreaming2", META = (DisplayName = "Media Capture Video Producer"))
-class PIXELSTREAMING2_API UPixelStreaming2VideoProducerMediaCapture : public UPixelStreaming2VideoProducer
+class PIXELSTREAMING2_API UPixelStreaming2VideoProducerMediaCapture : public UPixelStreaming2VideoProducerBase
 {
 	GENERATED_UCLASS_BODY()
 
@@ -50,7 +50,7 @@ public:
  * A video producer that streams the specified render target
  */
 UCLASS(NotBlueprintType, NotBlueprintable, Category = "PixelStreaming2", META = (DisplayName = "Render Target Video Producer"))
-class PIXELSTREAMING2_API UPixelStreaming2VideoProducerRenderTarget : public UPixelStreaming2VideoProducer
+class PIXELSTREAMING2_API UPixelStreaming2VideoProducerRenderTarget : public UPixelStreaming2VideoProducerBase
 {
 	GENERATED_UCLASS_BODY()
 
diff --git a/Engine/Plugins/Media/PixelStreaming2/Source/PixelStreaming2/Internal/ThreadSafeMap.h b/Engine/Plugins/Media/PixelStreaming2/Source/PixelStreaming2/Internal/ThreadSafeMap.h
--- a/Engine/Plugins/Media/PixelStreaming2/Source/PixelStreaming2/Internal/ThreadSafeMap.h
+++ b/Engine/Plugins/Media/PixelStreaming2/Source/PixelStreaming2/Internal/ThreadSafeMap.h
@@ -193,14 +193,14 @@ namespace UE::PixelStreaming2
 		{
 			FScopeLock Lock(&Mutex);
 
-			return FindChecked(Key);
+			return static_cast<Super*>(this)->FindChecked(Key);
 		}
 
 		FORCEINLINE const ValueType& operator[](KeyConstPointerType Key) const
 		{
 			FScopeLock Lock(&Mutex);
 
-			return FindChecked(Key);
+			return static_cast<Super*>(this)->FindChecked(Key);
 		}
 
 		// Note: Do not call a map method inside a visitor or you will deadlock
@@ -263,6 +263,49 @@ namespace UE::PixelStreaming2
 			return static_cast<Super*>(this)->Empty(ExpectedNumElements);
 		}
 
+		template <typename InSetKeyFuncs, typename InSetAllocator>
+		int32 GetKeys(TSet<KeyType, InSetKeyFuncs, InSetAllocator>& OutKeys)
+		{
+			FScopeLock Lock(&Mutex);
+
+			return static_cast<Super*>(this)->GetKeys(OutKeys);
+		}
+
+		[[nodiscard]] FORCEINLINE bool Contains(KeyConstPointerType Key)
+		{
+			FScopeLock Lock(&Mutex);
+
+			return static_cast<Super*>(this)->Contains(Key);
+		}
+
+		FORCEINLINE ValueType& Add(const KeyType& InKey, const ValueType& InValue)
+		{
+			FScopeLock Lock(&Mutex);
+
+			return static_cast<Super*>(this)->Emplace(InKey, InValue);
+		}
+
+		FORCEINLINE ValueType& Add(const KeyType& InKey, ValueType&& InValue)
+		{
+			FScopeLock Lock(&Mutex);
+
+			return static_cast<Super*>(this)->Emplace(InKey, MoveTempIfPossible(InValue));
+		}
+
+		FORCEINLINE ValueType& Add(KeyType&& InKey, const ValueType& InValue)
+		{
+			FScopeLock Lock(&Mutex);
+
+			return static_cast<Super*>(this)->Emplace(MoveTempIfPossible(InKey), InValue);
+		}
+
+		FORCEINLINE ValueType& Add(KeyType&& InKey, ValueType&& InValue)
+		{
+			FScopeLock Lock(&Mutex);
+
+			return static_cast<Super*>(this)->Emplace(MoveTempIfPossible(InKey), MoveTempIfPossible(InValue));
+		}
+
 	private:
 		mutable FCriticalSection Mutex;
 	};
diff --git a/Engine/Plugins/Media/PixelStreaming2/Source/PixelStreaming2/Internal/VideoCapturer.h b/Engine/Plugins/Media/PixelStreaming2/Source/PixelStreaming2/Internal/VideoCapturer.h
--- a/Engine/Plugins/Media/PixelStreaming2/Source/PixelStreaming2/Internal/VideoCapturer.h
+++ b/Engine/Plugins/Media/PixelStreaming2/Source/PixelStreaming2/Internal/VideoCapturer.h
@@ -31,8 +31,6 @@ namespace UE::PixelStreaming2
 
 		TSharedPtr<FVideoProducer> GetVideoProducer() { return VideoProducer; }
 
-		TRefCountPtr<EpicRtcVideoBufferInterface> GetFrameBuffer();
-
 		TSharedPtr<IPixelCaptureOutputFrame> RequestFormat(int32 Format, int32 LayerIndex = -1);
 
 		void ResetFrameCapturer();
@@ -44,7 +42,7 @@ namespace UE::PixelStreaming2
 		DECLARE_MULTICAST_DELEGATE(FOnFrameCaptured);
 		FOnFrameCaptured OnFrameCaptured;
 
-	private:
+	protected:
 		FVideoCapturer(TSharedPtr<FVideoProducer> VideoProducer);
 
 		int32 LastFrameWidth = -1;
diff --git a/Engine/Plugins/Media/PixelStreaming2/Source/PixelStreaming2/Internal/VideoSink.h b/Engine/Plugins/Media/PixelStreaming2/Source/PixelStreaming2/Internal/VideoSink.h
--- a/Engine/Plugins/Media/PixelStreaming2/Source/PixelStreaming2/Internal/VideoSink.h
+++ b/Engine/Plugins/Media/PixelStreaming2/Source/PixelStreaming2/Internal/VideoSink.h
@@ -6,9 +6,6 @@
 #include "CoreTypes.h"
 #include "HAL/CriticalSection.h"
 #include "IPixelStreaming2VideoSink.h"
-#include "RendererInterface.h"
-
-#include "epic_rtc/core/video/video_frame.h"
 
 namespace UE::PixelStreaming2
 {
@@ -22,8 +19,6 @@ namespace UE::PixelStreaming2
 
 		bool HasVideoConsumers();
 
-		void OnVideoData(const EpicRtcVideoFrame& Frame);
-
 		void SetMuted(bool bIsMuted);
 
 	protected:
@@ -32,13 +27,7 @@ namespace UE::PixelStreaming2
 
 		bool bIsMuted = false;
 
-		FCriticalSection				  RenderSyncContext;
-		FPooledRenderTargetDesc			  RenderTargetDescriptor;
-		TRefCountPtr<IPooledRenderTarget> RenderTarget;
-		TArray<uint8_t>					  Buffer;
-		FTextureRHIRef					  SourceTexture;
-
-	private:
+	protected:
 		void CallConsumeFrame(FTextureRHIRef Frame);
 	};
 
diff --git a/Engine/Plugins/Media/PixelStreaming2/Source/PixelStreaming2/Internal/VideoSource.h b/Engine/Plugins/Media/PixelStreaming2/Source/PixelStreaming2/Internal/VideoSource.h
new file mode 100644
--- /dev/null
+++ b/Engine/Plugins/Media/PixelStreaming2/Source/PixelStreaming2/Internal/VideoSource.h
@@ -0,0 +1,27 @@
+// Copyright Epic Games, Inc. All Rights Reserved.
+
+#pragma once
+
+#include "HAL/Event.h"
+#include "HAL/IConsoleManager.h"
+#include "HAL/Runnable.h"
+#include "HAL/RunnableThread.h"
+#include "Misc/SingleThreadRunnable.h"
+#include "Templates/SharedPointer.h"
+#include "VideoCapturer.h"
+
+namespace UE::PixelStreaming2
+{
+	class PIXELSTREAMING2_API FVideoSource
+	{
+	public:
+		virtual ~FVideoSource() = default;
+
+		virtual void PushFrame();
+		virtual void ForceKeyFrame();
+		virtual void SetMuted(bool bIsMuted);
+
+	protected:
+		bool bIsMuted = false;
+	};
+} // namespace UE::PixelStreaming2
diff --git a/Engine/Plugins/Media/PixelStreaming2/Source/PixelStreaming2/Internal/VideoSourceGroup.h b/Engine/Plugins/Media/PixelStreaming2/Source/PixelStreaming2/Internal/VideoSourceGroup.h
new file mode 100644
--- /dev/null
+++ b/Engine/Plugins/Media/PixelStreaming2/Source/PixelStreaming2/Internal/VideoSourceGroup.h
@@ -0,0 +1,69 @@
+// Copyright Epic Games, Inc. All Rights Reserved.
+
+#pragma once
+
+#include "HAL/Event.h"
+#include "HAL/IConsoleManager.h"
+#include "HAL/Runnable.h"
+#include "HAL/RunnableThread.h"
+#include "Misc/SingleThreadRunnable.h"
+#include "VideoCapturer.h"
+#include "Templates/SharedPointer.h"
+
+namespace UE::PixelStreaming2
+{
+	class FVideoSource;
+	class FFrameRunnable;
+
+	class PIXELSTREAMING2_API FVideoSourceGroup : public TSharedFromThis<FVideoSourceGroup>
+	{
+	public:
+		static TSharedPtr<FVideoSourceGroup> Create(TSharedPtr<FVideoCapturer> InVideoCapturer);
+		~FVideoSourceGroup();
+
+		void  SetFPS(int32 InFramesPerSecond);
+		int32 GetFPS();
+
+		void SetDecoupleFramerate(bool bDecouple);
+
+		void AddVideoSource(TSharedPtr<FVideoSource> VideoSource);
+		void RemoveVideoSource(const FVideoSource* ToRemove);
+		void RemoveAllVideoSources();
+
+		void Start();
+		void Stop();
+		void PushFrame();
+
+		void ForceKeyFrame();
+
+	private:
+		FVideoSourceGroup();
+
+		void StartThread();
+		void StopThread();
+
+		void OnFrameCaptured();
+
+		void OnWebRtcFpsChanged(IConsoleVariable* Var);
+		void OnDecoupleFramerateChanged(IConsoleVariable* Var);
+
+		bool bFPSOverride = false;
+		bool bDecoupleOverride = false;
+
+		bool  bRunning = false;
+		bool  bDecoupleFramerate = false;
+		int32 FramesPerSecond = 30;
+
+		TArray<TSharedPtr<FVideoSource>> VideoSources;
+
+		FDelegateHandle FrameDelegateHandle;
+		FDelegateHandle FpsDelegateHandle;
+		FDelegateHandle DecoupleDelegateHandle;
+
+		TSharedPtr<FRunnableThread> FrameThread = nullptr;	 // constant FPS tick thread
+		TSharedPtr<FFrameRunnable>	FrameRunnable = nullptr; // constant fps runnable
+
+		mutable FCriticalSection CriticalSection;
+		mutable FCriticalSection RunnableCriticalSection;
+	};
+} // namespace UE::PixelStreaming2
diff --git a/Engine/Plugins/Media/PixelStreaming2/Source/PixelStreaming2/PixelStreaming2.Build.cs b/Engine/Plugins/Media/PixelStreaming2/Source/PixelStreaming2/PixelStreaming2.Build.cs
--- a/Engine/Plugins/Media/PixelStreaming2/Source/PixelStreaming2/PixelStreaming2.Build.cs
+++ b/Engine/Plugins/Media/PixelStreaming2/Source/PixelStreaming2/PixelStreaming2.Build.cs
@@ -7,106 +7,37 @@ namespace UnrealBuildTool.Rules
 {
     public class PixelStreaming2 : ModuleRules
     {
-        string PixelStreaming2ProgramsOutputDirectory = "$(ProjectDir)/Samples/PixelStreaming2/WebServers";
-        const string PixelStreaming2ProgramsDirectory = "$(PluginDir)/Resources/WebServers";
-
-        private void AddFolder(string Folder)
-        {
-            string DirectoryToAdd = new DirectoryInfo(Path.Combine(PixelStreaming2ProgramsDirectory, Folder)).FullName;
-
-            if (!Directory.Exists(DirectoryToAdd))
-            {
-                return;
-            }
-
-            List<string> DependenciesToAdd = new List<string>();
-            DependenciesToAdd.AddRange(Directory.GetFiles(DirectoryToAdd, "*.*", SearchOption.AllDirectories));
-
-            List<string> PathsToIgnore = new List<string>();
-            PathsToIgnore.Add(new DirectoryInfo(DirectoryToAdd + "/node_modules").FullName);
-            PathsToIgnore.Add(new DirectoryInfo(DirectoryToAdd + "/implementations/typescript/node_modules").FullName);
-            PathsToIgnore.Add(new DirectoryInfo(DirectoryToAdd + "/implementations/EpicGames/node_modules").FullName);
-            PathsToIgnore.Add(new DirectoryInfo(DirectoryToAdd + "/implementations/react/node_modules").FullName);
-            PathsToIgnore.Add(new DirectoryInfo(DirectoryToAdd + "/library/node_modules").FullName);
-            PathsToIgnore.Add(new DirectoryInfo(DirectoryToAdd + "/ui-library/node_modules").FullName);
-            PathsToIgnore.Add(new DirectoryInfo(DirectoryToAdd + "/logs").FullName);
-            PathsToIgnore.Add(new DirectoryInfo(DirectoryToAdd + "/platform_scripts/cmd/node").FullName);
-            PathsToIgnore.Add(new DirectoryInfo(DirectoryToAdd + "/platform_scripts/cmd/coturn").FullName);
-            PathsToIgnore.Add(new DirectoryInfo(DirectoryToAdd + "/platform_scripts/bash/node").FullName);
-
-            // Go through each file/folder in the WebServers and selectively choose which ones we wish to package.
-            foreach (var DependencySource in DependenciesToAdd)
-            {
-                // We want to skip packaging the above directories as they are all generated by running scripts
-                bool bSkip = false;
-                foreach (var PathToIgnore in PathsToIgnore)
-                {
-                    if (DependencySource.StartsWith(PathToIgnore))
-                    {
-                        bSkip = true;
-                        continue;
-                    }
-                }
-
-                // If we are not skipping, we want to add this to our packaged output
-                if (!bSkip)
-                {
-                    string RelativeSourcePath = Path.GetRelativePath(PixelStreaming2ProgramsDirectory, DependencySource);
-                    RuntimeDependencies.Add(Path.Combine(PixelStreaming2ProgramsOutputDirectory, RelativeSourcePath), DependencySource, StagedFileType.SystemNonUFS);
-                }
-            }
-        }
-
         public PixelStreaming2(ReadOnlyTargetRules Target) : base(Target)
         {
             // This is so for game projects using our public headers don't have to include extra modules they might not know about.
             PublicDependencyModuleNames.AddRange(new string[]
             {
-                "ApplicationCore",
-                "InputDevice",
-                "WebRTC",
-                "PixelCapture",
-                "PixelStreaming2Input",
-                "MediaIOCore",
-                "EpicRtc",
-                "AVCodecsCore"
-            });
-
-            PrivateDependencyModuleNames.AddRange(new string[]
-            {
-                "Core",
+                "AudioMixer",
+                "AudioPlatformConfiguration",
+                "AVCodecsCore",
+                "AVCodecsCoreRHI",
                 "CoreUObject",
-                "Engine",
-                "EngineSettings",
-                "InputCore",
                 "Json",
+                "MediaIOCore",
+                "MediaUtils",
+                "PixelCapture",
+                "PixelStreaming2Input",
+                "RHI",
                 "Renderer",
                 "RenderCore",
-                "RHI",
                 "SignalProcessing",
                 "Slate",
                 "SlateCore",
-                "AudioMixer",
-                "WebRTC",
-                "WebSockets",
-                "Sockets",
-                "MediaUtils",
-                "DeveloperSettings",
-                "AVCodecsCore",
-                "AVCodecsCoreRHI",
-                "PixelCaptureShaders",
-                "PixelStreaming2Servers",
-                "PixelStreaming2HMD",
                 "TraceLog",
-                "HTTP",
-                "NVML",
-                "EpicRtc",
-                "AudioPlatformConfiguration",
-                "LibVpx",
-                "PixelStreaming2Core"
+                "WebRTC" // Used by pixelcapture :facepalm: 
             });
 
-            PrivateDefinitions.Add("PIXELSTREAMING2_DUMP_ENCODING=0");
+            PrivateDependencyModuleNames.AddRange(new string[]
+            {
+                "Core",
+                "Engine",
+                "PixelStreaming2Settings"
+            });
 
             if (Target.IsInPlatformGroup(UnrealPlatformGroup.Windows) || Target.IsInPlatformGroup(UnrealPlatformGroup.Linux))
             {
@@ -114,8 +45,6 @@ namespace UnrealBuildTool.Rules
                 AddEngineThirdPartyPrivateStaticDependencies(Target, "Vulkan", "CUDA");
             }
 
-            PrivateDefinitions.Add("PIXELSTREAMING2_DUMP_ENCODING=0");
-
 
             if (Target.IsInPlatformGroup(UnrealPlatformGroup.Windows))
             {
@@ -129,22 +58,6 @@ namespace UnrealBuildTool.Rules
             {
                 AddEngineThirdPartyPrivateStaticDependencies(Target, "MetalCPP");
             }
-
-            // When we build a Game target we also package the servers with it as runtime dependencies
-            if (Target.Type == TargetType.Game && Target.ProjectFile != null)
-            {
-                AddFolder("Frontend");
-                AddFolder("SignallingWebServer");
-                AddFolder("Matchmaker");
-                AddFolder("SFU");
-                AddFolder("Extras");
-                AddFolder("Signalling");
-                AddFolder("Common");
-                AddFolder("SS_Test");
-
-                RuntimeDependencies.Add("$(ProjectDir)/Samples/PixelStreaming2/WebServers/get_ps_servers.bat", "$(PluginDir)/Resources/WebServers/get_ps_servers.bat", StagedFileType.NonUFS);
-                RuntimeDependencies.Add("$(ProjectDir)/Samples/PixelStreaming2/WebServers/get_ps_servers.sh", "$(PluginDir)/Resources/WebServers/get_ps_servers.sh", StagedFileType.NonUFS);
-            }
         }
     }
 }
diff --git a/Engine/Plugins/Media/PixelStreaming2/Source/PixelStreaming2/Private/AudioCapturer.cpp b/Engine/Plugins/Media/PixelStreaming2/Source/PixelStreaming2/Private/AudioCapturer.cpp
new file mode 100644
--- /dev/null
+++ b/Engine/Plugins/Media/PixelStreaming2/Source/PixelStreaming2/Private/AudioCapturer.cpp
@@ -0,0 +1,283 @@
+// Copyright Epic Games, Inc. All Rights Reserved.
+
+#include "AudioCapturer.h"
+
+#include "AudioDevice.h"
+#include "Engine/Engine.h"
+#include "Logging.h"
+#include "Misc/CoreDelegates.h"
+#include "PixelStreaming2PluginSettings.h"
+#include "Sound/SampleBufferIO.h"
+
+namespace UE::PixelStreaming2
+{
+	FAudioPatchMixer::FAudioPatchMixer(uint8 NumChannels, uint32 SampleRate, float SampleSizeSeconds)
+		: NumChannels(NumChannels)
+		, SampleRate(SampleRate)
+		, SampleSizeSeconds(SampleSizeSeconds)
+	{
+	}
+
+	uint32 FAudioPatchMixer::GetMaxBufferSize() const
+	{
+		return NumChannels * SampleRate * SampleSizeSeconds;
+	}
+
+	uint8 FAudioPatchMixer::GetNumChannels() const
+	{
+		return NumChannels;
+	}
+
+	uint32 FAudioPatchMixer::GetSampleRate() const
+	{
+		return SampleRate;
+	}
+
+	/***************************************************
+	 *
+	 ****************************************************/
+	FPatchInputProxy::FPatchInputProxy(TSharedPtr<FAudioPatchMixer> InMixer)
+		: Mixer(InMixer)
+		// We don't want the patch input to handle gain as the capturer handles that at the end
+		, PatchInput(Mixer->AddNewInput(Mixer->GetMaxBufferSize(), 1.0f))
+		, NumChannels(Mixer->GetNumChannels())
+		, SampleRate(Mixer->GetSampleRate())
+	{
+	}
+
+	FPatchInputProxy::~FPatchInputProxy()
+	{
+		Mixer->RemovePatch(PatchInput);
+	}
+
+	void FPatchInputProxy::PushAudio(const float* AudioData, int32 InNumSamples, int32 InNumChannels, int32 InSampleRate)
+	{
+		TArray<float> AudioBuffer;
+
+		if (SampleRate != InSampleRate)
+		{
+			float SampleRateConversionRatio = static_cast<float>(SampleRate) / static_cast<float>(InSampleRate);
+			Resampler.Init(Audio::EResamplingMethod::Linear, SampleRateConversionRatio, InNumChannels);
+
+			int32 NumConvertedSamples = InNumSamples / SampleRateConversionRatio;
+			int32 OutputSamples = INDEX_NONE;
+			AudioBuffer.AddZeroed(NumConvertedSamples);
+			// Perform the sample rate conversion
+			int32 ErrorCode = Resampler.ProcessAudio(const_cast<float*>(AudioData), InNumSamples, false, AudioBuffer.GetData(), AudioBuffer.Num(), OutputSamples);
+			verifyf(OutputSamples <= NumConvertedSamples, TEXT("OutputSamples > NumConvertedSamples"));
+			if (ErrorCode != 0)
+			{
+				UE_LOG(LogPixelStreaming2, Warning, TEXT("(FAudioInput) Problem occured resampling audio data. Code: %d"), ErrorCode);
+				return;
+			}
+		}
+		else
+		{
+			AudioBuffer.AddZeroed(InNumSamples);
+			FMemory::Memcpy(AudioBuffer.GetData(), AudioData, AudioBuffer.Num() * sizeof(float));
+		}
+
+		// Note: TSampleBuffer takes in AudioData as float* and internally converts to int16
+		Audio::TSampleBuffer<int16_t> Buffer(AudioBuffer.GetData(), AudioBuffer.Num(), InNumChannels, SampleRate);
+		if (InNumChannels != NumChannels)
+		{
+			if (InNumChannels < NumChannels)
+			{
+				// Up mix by duplicating the mono source to each channel
+				TArrayView<int16> SourceBuffer = Buffer.GetArrayView();
+				TArray<int16>	  MixedBuffer;
+				MixedBuffer.SetNumZeroed(InNumSamples / InNumChannels * NumChannels);
+				for (int32 SrcSampleIdx = 0; SrcSampleIdx < Buffer.GetNumSamples(); SrcSampleIdx++)
+				{
+					int32 DestSampleIdx = SrcSampleIdx * NumChannels;
+					for (int32 Channel = 0; Channel < NumChannels; Channel++)
+					{
+						MixedBuffer[DestSampleIdx + Channel] = SourceBuffer[SrcSampleIdx];
+					}
+				}
+
+				Buffer.CopyFrom(MixedBuffer, NumChannels, InSampleRate);
+			}
+			else
+			{
+				// Down mix using inbuilt method
+				Buffer.MixBufferToChannels(NumChannels);
+			}
+		}
+
+		// Apply gain
+		float Gain = UPixelStreaming2PluginSettings::CVarWebRTCAudioGain.GetValueOnAnyThread();
+		if (Gain != 1.0f)
+		{
+			int16* PCMAudio = Buffer.GetArrayView().GetData();
+
+			// multiply audio by gain multiplier
+			for (int i = 0; i < Buffer.GetNumSamples(); i++)
+			{
+				*PCMAudio = FMath::Max(-32768, FMath::Min(32767, *PCMAudio * Gain));
+				PCMAudio++;
+			}
+		}
+
+		Audio::TSampleBuffer<float> BufferToPush(Buffer.GetData(), Buffer.GetNumSamples(), Buffer.GetNumChannels(), Buffer.GetSampleRate());
+		PatchInput.PushAudio(BufferToPush.GetData(), BufferToPush.GetNumSamples());
+	}
+
+	/***************************************************
+	 *
+	 ****************************************************/
+	FMixAudioTask::FMixAudioTask(FAudioCapturer* Capturer, TSharedPtr<FAudioPatchMixer> Mixer)
+		: Capturer(Capturer)
+		, Mixer(Mixer)
+	{
+		MixingBuffer.SetNumUninitialized(Mixer->GetMaxBufferSize());
+	}
+
+	void FMixAudioTask::Tick(float DeltaMs)
+	{
+		if (!Mixer)
+		{
+			return;
+		}
+
+		// 4 samples is the absolute minimum required for mixing
+		if (MixingBuffer.Num() < 4)
+		{
+			return;
+		}
+
+		int32 TargetNumSamples = Mixer->MaxNumberOfSamplesThatCanBePopped();
+		if (TargetNumSamples < 0)
+		{
+			return;
+		}
+
+		int32 NSamplesPopped = Mixer->PopAudio(MixingBuffer.GetData(), TargetNumSamples, false /* bUseLatestAudio */);
+		if (NSamplesPopped == 0)
+		{
+			return;
+		}
+
+		Capturer->OnAudio(MixingBuffer.GetData(), NSamplesPopped, Mixer->GetNumChannels(), Mixer->GetSampleRate());
+	}
+	const FString& FMixAudioTask::GetName() const
+	{
+		static FString TaskName = TEXT("MixAudioTask");
+		return TaskName;
+	}
+
+	/***************************************************
+	 *
+	 ****************************************************/
+	TSharedPtr<FAudioCapturer> FAudioCapturer::Create(const int InSampleRate, const int InNumChannels, const float InSampleSizeInSeconds)
+	{
+		TSharedPtr<FAudioCapturer> AudioCapturer(new FAudioCapturer(InSampleRate, InNumChannels, InSampleSizeInSeconds));
+
+		FAudioDeviceManagerDelegates::OnAudioDeviceCreated.AddSP(AudioCapturer.ToSharedRef(), &FAudioCapturer::CreateAudioProducer);
+		FAudioDeviceManagerDelegates::OnAudioDeviceDestroyed.AddSP(AudioCapturer.ToSharedRef(), &FAudioCapturer::RemoveAudioProducer);
+
+		if (UPixelStreaming2PluginSettings::FDelegates* Delegates = UPixelStreaming2PluginSettings::Delegates())
+		{
+			Delegates->OnDebugDumpAudioChanged.AddSP(AudioCapturer.ToSharedRef(), &FAudioCapturer::OnDebugDumpAudioChanged);
+
+			TWeakPtr<FAudioCapturer> WeakAudioMixingCapturer = AudioCapturer;
+			FCoreDelegates::OnEnginePreExit.AddLambda([WeakAudioMixingCapturer]() {
+				if (TSharedPtr<FAudioCapturer> AudioCapturer = WeakAudioMixingCapturer.Pin())
+				{
+					AudioCapturer->OnEnginePreExit();
+				}
+			});
+		}
+
+		return AudioCapturer;
+	}
+
+	FAudioCapturer::FAudioCapturer(const int SampleRate, const int NumChannels, const float NumSampleSizeInSeconds)
+		: SampleRate(SampleRate)
+		, NumChannels(NumChannels)
+		, SampleSizeSeconds(NumSampleSizeInSeconds)
+	{
+		// subscribe to audio data
+		if (!GEngine)
+		{
+			// No engine. Possibly running editor tests
+			return;
+		}
+		Mixer = MakeShared<FAudioPatchMixer>(NumChannels, SampleRate, SampleSizeSeconds);
+		MixerTask = FPixelStreamingTickableTask::Create<FMixAudioTask>(this, Mixer);
+		FAudioDeviceManager::Get()->IterateOverAllDevices([this](Audio::FDeviceId AudioDeviceId, FAudioDevice*) {
+			CreateAudioProducer(AudioDeviceId);
+		});
+	}
+
+	// Mixed audio input will push its audio to an FPatchInputProxy for mixing
+	TSharedPtr<FAudioProducer> FAudioCapturer::CreateAudioProducer()
+	{
+		// The lifetimes of audio producers created by the user are the responsibility of the user
+		return FAudioProducer::Create(MakeShared<FPatchInputProxy>(Mixer));
+	}
+
+	void FAudioCapturer::CreateAudioProducer(Audio::FDeviceId AudioDeviceId)
+	{
+		// The lifetimes of audio producers created by the engine are our responsibility
+		TSharedPtr<FAudioProducer> AudioInput = FAudioProducer::Create(AudioDeviceId, MakeShared<FPatchInputProxy>(Mixer));
+		AudioProducers.Add(AudioDeviceId, AudioInput);
+	}
+
+	void FAudioCapturer::RemoveAudioProducer(Audio::FDeviceId AudioDeviceId)
+	{
+		AudioProducers.Remove(AudioDeviceId);
+	}
+
+	void FAudioCapturer::OnAudio(const float* AudioData, int32 InNumSamples, int32 InNumChannels, int32 InSampleRate)
+	{
+		// Note: TSampleBuffer takes in AudioData as float* and internally converts to int16
+		Audio::TSampleBuffer<int16_t> Buffer(AudioData, InNumSamples, InNumChannels, SampleRate);
+
+		if (UPixelStreaming2PluginSettings::CVarDebugDumpAudio.GetValueOnAnyThread())
+		{
+			DebugDumpAudioBuffer.Append(Buffer.GetData(), Buffer.GetNumSamples(), Buffer.GetNumChannels(), Buffer.GetSampleRate());
+		}
+
+		PushAudio(AudioData, InNumSamples, InNumChannels, InSampleRate);
+	}
+
+	void FAudioCapturer::PushAudio(const float* AudioData, int32 InNumSamples, int32 InNumChannels, int32 InSampleRate)
+	{
+		Audio::TSampleBuffer<int16_t> Buffer(AudioData, InNumSamples, InNumChannels, SampleRate);
+
+		OnAudioBuffer.Broadcast(Buffer.GetData(), Buffer.GetNumSamples(), Buffer.GetNumChannels(), Buffer.GetSampleRate());
+	}
+
+	void FAudioCapturer::OnDebugDumpAudioChanged(IConsoleVariable* Var)
+	{
+		if (!Var->GetBool())
+		{
+			WriteDebugAudio();
+		}
+	}
+
+	void FAudioCapturer::OnEnginePreExit()
+	{
+		// If engine is exiting but the dump cvar is true, we need to manually trigger a write
+		if (UPixelStreaming2PluginSettings::CVarDebugDumpAudio.GetValueOnAnyThread())
+		{
+			WriteDebugAudio();
+		}
+	}
+
+	void FAudioCapturer::WriteDebugAudio()
+	{
+		// Only write audio if we actually have some
+		if (DebugDumpAudioBuffer.GetSampleDuration() <= 0.f)
+		{
+			return;
+		}
+
+		Audio::FSoundWavePCMWriter Writer;
+		FString					   FilePath = TEXT("");
+		Writer.SynchronouslyWriteToWavFile(DebugDumpAudioBuffer, TEXT("PixelStreamingMixedAudio"), TEXT(""), &FilePath);
+		UE_LOGFMT(LogPixelStreaming2, Log, "Saving audio sample to: {0}", FilePath);
+		DebugDumpAudioBuffer.Reset();
+	}
+} // namespace UE::PixelStreaming2
\ No newline at end of file
diff --git a/Engine/Plugins/Media/PixelStreaming2/Source/PixelStreaming2/Private/AudioProducer.cpp b/Engine/Plugins/Media/PixelStreaming2/Source/PixelStreaming2/Private/AudioProducer.cpp
new file mode 100644
--- /dev/null
+++ b/Engine/Plugins/Media/PixelStreaming2/Source/PixelStreaming2/Private/AudioProducer.cpp
@@ -0,0 +1,52 @@
+// Copyright Epic Games, Inc. All Rights Reserved.
+
+#include "AudioProducer.h"
+
+#include "AudioCapturer.h"
+#include "AudioDevice.h"
+#include "AudioDeviceManager.h"
+
+namespace UE::PixelStreaming2
+{
+	TSharedPtr<FAudioProducer> FAudioProducer::Create(Audio::FDeviceId InAudioDeviceId, TSharedPtr<FPatchInputProxy> InPatchInput)
+	{
+		TSharedPtr<FAudioProducer> Listener = TSharedPtr<FAudioProducer>(new FAudioProducer(InPatchInput));
+		if (FAudioDevice* AudioDevice = FAudioDeviceManager::Get()->GetAudioDeviceRaw(InAudioDeviceId))
+		{
+			AudioDevice->RegisterSubmixBufferListener(Listener.ToSharedRef(), AudioDevice->GetMainSubmixObject());
+		}
+
+		return Listener;
+	}
+
+	TSharedPtr<FAudioProducer> FAudioProducer::Create(TSharedPtr<FPatchInputProxy> InPatchInput)
+	{
+		TSharedPtr<FAudioProducer> Listener = TSharedPtr<FAudioProducer>(new FAudioProducer(InPatchInput));
+		return Listener;
+	}
+
+	FAudioProducer::FAudioProducer(TSharedPtr<FPatchInputProxy> PatchInput)
+		: PatchInput(PatchInput)
+	{
+	}
+
+	void FAudioProducer::PushAudio(const float* AudioData, int32 NumSamples, int32 NumChannels, int32 SampleRate)
+	{
+		if (bIsMuted)
+		{
+			return;
+		}
+
+		PatchInput->PushAudio(AudioData, NumSamples, NumChannels, SampleRate);
+	}
+
+	void FAudioProducer::OnNewSubmixBuffer(const USoundSubmix* OwningSubmix, float* AudioData, int32 NumSamples, int32 NumChannels, const int32 SampleRate, double AudioClock)
+	{
+		if (bIsMuted)
+		{
+			return;
+		}
+
+		PatchInput->PushAudio(AudioData, NumSamples, NumChannels, SampleRate);
+	}
+} // namespace UE::PixelStreaming2
\ No newline at end of file
diff --git a/Engine/Plugins/Media/PixelStreaming2/Source/PixelStreaming2/Private/Blueprints/PixelStreaming2Blueprints.cpp b/Engine/Plugins/Media/PixelStreaming2/Source/PixelStreaming2/Private/Blueprints/PixelStreaming2Blueprints.cpp
--- a/Engine/Plugins/Media/PixelStreaming2/Source/PixelStreaming2/Private/Blueprints/PixelStreaming2Blueprints.cpp
+++ b/Engine/Plugins/Media/PixelStreaming2/Source/PixelStreaming2/Private/Blueprints/PixelStreaming2Blueprints.cpp
@@ -10,13 +10,7 @@
 
 TSharedPtr<IPixelStreaming2Streamer> GetStreamer(FString StreamerId)
 {
-	IPixelStreaming2Module* Module = UE::PixelStreaming2::FPixelStreaming2Module::GetModule();
-	if (!Module)
-	{
-		UE_LOG(LogPixelStreaming2, Error, TEXT("Could not get Pixel Streaming module, so unable to get streamer by id."));
-		return nullptr;
-	}
-	TSharedPtr<IPixelStreaming2Streamer> Streamer = Module->FindStreamer(StreamerId);
+	TSharedPtr<IPixelStreaming2Streamer> Streamer = IPixelStreaming2Module::Get().FindStreamer(StreamerId);
 	if (Streamer)
 	{
 		return Streamer;
@@ -136,13 +130,7 @@ TArray<FString> UPixelStreaming2Blueprints::StreamerGetConnectedPlayers(FString
 
 FString UPixelStreaming2Blueprints::GetDefaultStreamerID()
 {
-	IPixelStreaming2Module* Module = UE::PixelStreaming2::FPixelStreaming2Module::GetModule();
-	if (!Module)
-	{
-		UE_LOG(LogPixelStreaming2, Error, TEXT("Pixel Streaming module not initialized yet - cannot get default streamer id."));
-		return FString(TEXT("ModuleNotInitializedYet"));
-	}
-	return Module->GetDefaultStreamerID();
+	return IPixelStreaming2Module::Get().GetDefaultStreamerID();
 }
 
 UPixelStreaming2Delegates* UPixelStreaming2Blueprints::GetDelegates()
diff --git a/Engine/Plugins/Media/PixelStreaming2/Source/PixelStreaming2/Private/Blueprints/PixelStreaming2InputComponent.cpp b/Engine/Plugins/Media/PixelStreaming2/Source/PixelStreaming2/Private/Blueprints/PixelStreaming2InputComponent.cpp
--- a/Engine/Plugins/Media/PixelStreaming2/Source/PixelStreaming2/Private/Blueprints/PixelStreaming2InputComponent.cpp
+++ b/Engine/Plugins/Media/PixelStreaming2/Source/PixelStreaming2/Private/Blueprints/PixelStreaming2InputComponent.cpp
@@ -54,22 +54,15 @@ void UPixelStreaming2Input::EndPlay(const EEndPlayReason::Type EndPlayReason)
 
 void UPixelStreaming2Input::SendPixelStreaming2Response(const FString& Descriptor)
 {
-	if (PixelStreaming2Module)
-	{
-		PixelStreaming2Module->ForEachStreamer([&Descriptor, this](TSharedPtr<IPixelStreaming2Streamer> Streamer) {
-			TSharedPtr<IPixelStreaming2InputHandler> Handler = Streamer->GetInputHandler().Pin();
-			if (!Handler)
-			{
-				UE_LOG(LogPixelStreaming2, Error, TEXT("Pixel Streaming input handler was null when sending response message."));
-				return;
-			}
-			Streamer->SendAllPlayersMessage(EPixelStreaming2FromStreamerMessage::Response, Descriptor);
-		});
-	}
-	else
-	{
-		UE_LOG(LogPixelStreaming2, Warning, TEXT("Pixel Streaming input component skipped sending response. This is expected on dedicated servers."));
-	}
+	IPixelStreaming2Module::Get().ForEachStreamer([&Descriptor, this](TSharedPtr<IPixelStreaming2Streamer> Streamer) {
+		TSharedPtr<IPixelStreaming2InputHandler> Handler = Streamer->GetInputHandler().Pin();
+		if (!Handler)
+		{
+			UE_LOG(LogPixelStreaming2, Error, TEXT("Pixel Streaming input handler was null when sending response message."));
+			return;
+		}
+		Streamer->SendAllPlayersMessage(EPixelStreaming2FromStreamerMessage::Response, Descriptor);
+	});
 }
 
 void UPixelStreaming2Input::GetJsonStringValue(FString Descriptor, FString FieldName, FString& StringValue, bool& Success)
diff --git a/Engine/Plugins/Media/PixelStreaming2/Source/PixelStreaming2/Private/Blueprints/PixelStreaming2StreamerComponent.cpp b/Engine/Plugins/Media/PixelStreaming2/Source/PixelStreaming2/Private/Blueprints/PixelStreaming2StreamerComponent.cpp
--- a/Engine/Plugins/Media/PixelStreaming2/Source/PixelStreaming2/Private/Blueprints/PixelStreaming2StreamerComponent.cpp
+++ b/Engine/Plugins/Media/PixelStreaming2/Source/PixelStreaming2/Private/Blueprints/PixelStreaming2StreamerComponent.cpp
@@ -1,6 +1,7 @@
 // Copyright Epic Games, Inc. All Rights Reserved.
 
 #include "PixelStreaming2StreamerComponent.h"
+
 #include "IPixelStreaming2Module.h"
 #include "VideoProducerMediaCapture.h"
 #include "Engine/GameEngine.h"
@@ -64,12 +65,12 @@ void UPixelStreaming2StreamerComponent::StartStreaming()
 	if (UsePixelStreamingURL)
 	{
 		FString ServerURL;
-		FParse::Value(FCommandLine::Get(), TEXT("PixelStreaming2URL="), ServerURL);
-		Streamer->SetSignallingServerURL(ServerURL);
+		FParse::Value(FCommandLine::Get(), TEXT("PixelStreamingURL="), ServerURL);
+		Streamer->SetConnectionURL(ServerURL);
 	}
 	else
 	{
-		Streamer->SetSignallingServerURL(SignallingServerURL);
+		Streamer->SetConnectionURL(SignallingServerURL);
 	}
 	Streamer->StartStreaming();
 }
diff --git a/Engine/Plugins/Media/PixelStreaming2/Source/PixelStreaming2/Private/Blueprints/PixelStreaming2StreamerComponent.h b/Engine/Plugins/Media/PixelStreaming2/Source/PixelStreaming2/Private/Blueprints/PixelStreaming2StreamerComponent.h
--- a/Engine/Plugins/Media/PixelStreaming2/Source/PixelStreaming2/Private/Blueprints/PixelStreaming2StreamerComponent.h
+++ b/Engine/Plugins/Media/PixelStreaming2/Source/PixelStreaming2/Private/Blueprints/PixelStreaming2StreamerComponent.h
@@ -69,7 +69,7 @@ public:
 	bool CoupleFramerate = false;
 
 	UPROPERTY(EditAnywhere, Category = "PixelStreaming2")
-	TObjectPtr<UPixelStreaming2VideoProducer> VideoProducer = nullptr;
+	TObjectPtr<UPixelStreaming2VideoProducerBase> VideoProducer = nullptr;
 
 private:
 	TSharedPtr<IPixelStreaming2Streamer> Streamer;
diff --git a/Engine/Plugins/Media/PixelStreaming2/Source/PixelStreaming2/Private/Blueprints/PixelStreaming2VideoProducer.cpp b/Engine/Plugins/Media/PixelStreaming2/Source/PixelStreaming2/Private/Blueprints/PixelStreaming2VideoProducer.cpp
--- a/Engine/Plugins/Media/PixelStreaming2/Source/PixelStreaming2/Private/Blueprints/PixelStreaming2VideoProducer.cpp
+++ b/Engine/Plugins/Media/PixelStreaming2/Source/PixelStreaming2/Private/Blueprints/PixelStreaming2VideoProducer.cpp
@@ -8,9 +8,9 @@
 #include "VideoProducerRenderTarget.h"
 
 /**
- * ---------- UPixelStreaming2VideoProducer -------------------
+ * ---------- UPixelStreaming2VideoProducerBase -------------------
  */
-UPixelStreaming2VideoProducer::UPixelStreaming2VideoProducer(const FObjectInitializer& ObjectInitializer)
+UPixelStreaming2VideoProducerBase::UPixelStreaming2VideoProducerBase(const FObjectInitializer& ObjectInitializer)
 	: Super(ObjectInitializer)
 {
 }
diff --git a/Engine/Plugins/Media/PixelStreaming2/Source/PixelStreaming2/Private/PixelStreaming2MediaIOCapture.cpp b/Engine/Plugins/Media/PixelStreaming2/Source/PixelStreaming2/Private/PixelStreaming2MediaIOCapture.cpp
--- a/Engine/Plugins/Media/PixelStreaming2/Source/PixelStreaming2/Private/PixelStreaming2MediaIOCapture.cpp
+++ b/Engine/Plugins/Media/PixelStreaming2/Source/PixelStreaming2/Private/PixelStreaming2MediaIOCapture.cpp
@@ -5,7 +5,6 @@
 #include "Logging.h"
 #include "MediaShaders.h"
 #include "PixelCaptureInputFrameRHI.h"
-#include "PixelStreaming2Module.h"
 #include "RenderGraphUtils.h"
 #include "RHI.h"
 #include "ScreenPass.h"
diff --git a/Engine/Plugins/Media/PixelStreaming2/Source/PixelStreaming2/Private/PixelStreaming2Module.cpp b/Engine/Plugins/Media/PixelStreaming2/Source/PixelStreaming2/Private/PixelStreaming2Module.cpp
--- a/Engine/Plugins/Media/PixelStreaming2/Source/PixelStreaming2/Private/PixelStreaming2Module.cpp
+++ b/Engine/Plugins/Media/PixelStreaming2/Source/PixelStreaming2/Private/PixelStreaming2Module.cpp
@@ -1,118 +1,39 @@
 // Copyright Epic Games, Inc. All Rights Reserved.
 
+#pragma once
+
 #include "PixelStreaming2Module.h"
 
 #include "Blueprints/PixelStreaming2InputComponent.h"
-#include "CoderUtils.h"
-#include "CoreMinimal.h"
-#include "CoreUtils.h"
-#include "Engine/GameViewportClient.h"
-#include "Engine/Texture2D.h"
-#include "EpicRtcThread.h"
-#include "IPixelStreaming2InputModule.h"
+#include "Engine/Engine.h"
+#include "IPixelStreaming2Streamer.h"
 #include "Logging.h"
-#include "Modules/ModuleManager.h"
 #include "PixelStreaming2Delegates.h"
-#include "PixelStreaming2Utils.h"
 #include "PixelStreaming2PluginSettings.h"
-#include "Slate/SceneViewport.h"
-#include "Streamer.h"
-#include "UObject/UObjectIterator.h"
-#include "UtilsCommon.h"
-
-#if PLATFORM_LINUX
-	#include "CudaModule.h"
-#endif
-
-#if PLATFORM_WINDOWS
-	#include "Windows/AllowWindowsPlatformTypes.h"
-THIRD_PARTY_INCLUDES_START
-	#include <VersionHelpers.h>
-THIRD_PARTY_INCLUDES_END
-	#include "Windows/HideWindowsPlatformTypes.h"
-#endif
-
-#include "AudioDeviceManager.h"
-#include "RenderingThread.h"
-#include "Misc/CommandLine.h"
-#include "Misc/Parse.h"
-#include "RendererInterface.h"
-#include "Rendering/SlateRenderer.h"
-#include "Framework/Application/SlateApplication.h"
-#include "Misc/ConfigCacheIni.h"
-#include "GameFramework/GameModeBase.h"
-#include "Dom/JsonObject.h"
-#include "Misc/App.h"
-#include "Misc/MessageDialog.h"
-#include "Async/Async.h"
-#include "Engine/Engine.h"
-#include "WebSocketsModule.h"
-
-#include "Video/Encoders/Configs/VideoEncoderConfigH264.h"
+#include "RHI.h"
+#include "Thread.h"
+#include "UtilsCoder.h"
+#include "UtilsCore.h"
+#include "VideoProducer.h"
 #include "Video/Encoders/Configs/VideoEncoderConfigAV1.h"
-
-#if !UE_BUILD_SHIPPING
-	#include "DrawDebugHelpers.h"
-#endif
-
-#include "VideoProducerBackBuffer.h"
-#include "VideoProducerMediaCapture.h"
-#include "Engine/GameEngine.h"
-#include "Stats.h"
-#include "UtilsString.h"
-
-#include "EpicRtcAllocator.h"
-#include "EpicRtcAudioMixingCapturer.h"
-#include "EpicRtcLogging.h"
-#include "EpicRtcVideoEncoderInitializer.h"
-#include "EpicRtcVideoDecoderInitializer.h"
-#include "EpicRtcWebsocketFactory.h"
+#include "Video/Encoders/Configs/VideoEncoderConfigH264.h"
 
 namespace UE::PixelStreaming2
 {
-	constexpr EpicRtcLogLevel UnrealLogToEpicRtcCategoryMap[] = {
-		EpicRtcLogLevel::Off,
-		EpicRtcLogLevel::Critical,
-		EpicRtcLogLevel::Error,
-		EpicRtcLogLevel::Warning,
-		EpicRtcLogLevel::Info,
-		EpicRtcLogLevel::Info,
-		EpicRtcLogLevel::Debug,
-		EpicRtcLogLevel::Trace,
-		EpicRtcLogLevel::Trace
-	};
-
-	static_assert(EpicRtcLogLevel::Off == UnrealLogToEpicRtcCategoryMap[ELogVerbosity::Type::NoLogging]);
-	static_assert(EpicRtcLogLevel::Critical == UnrealLogToEpicRtcCategoryMap[ELogVerbosity::Type::Fatal]);
-	static_assert(EpicRtcLogLevel::Error == UnrealLogToEpicRtcCategoryMap[ELogVerbosity::Type::Error]);
-	static_assert(EpicRtcLogLevel::Warning == UnrealLogToEpicRtcCategoryMap[ELogVerbosity::Type::Warning]);
-	static_assert(EpicRtcLogLevel::Info == UnrealLogToEpicRtcCategoryMap[ELogVerbosity::Type::Display]);
-	static_assert(EpicRtcLogLevel::Info == UnrealLogToEpicRtcCategoryMap[ELogVerbosity::Type::Log]);
-	static_assert(EpicRtcLogLevel::Debug == UnrealLogToEpicRtcCategoryMap[ELogVerbosity::Type::Verbose]);
-	static_assert(EpicRtcLogLevel::Trace == UnrealLogToEpicRtcCategoryMap[ELogVerbosity::Type::VeryVerbose]);
-	static_assert(EpicRtcLogLevel::Trace == UnrealLogToEpicRtcCategoryMap[ELogVerbosity::Type::All]);
-
 	FPixelStreaming2Module* FPixelStreaming2Module::PixelStreaming2Module = nullptr;
 
-	FUtf8String FPixelStreaming2Module::EpicRtcConferenceName("pixel_streaming_conference_instance");
-
-	/**
-	 * Stats logger - as turned on/off by CVarPixelStreaming2LogStats
-	 */
-	void ConsumeStat(FString PlayerId, FName StatName, float StatValue)
+	FPixelStreaming2Module* FPixelStreaming2Module::GetModule()
 	{
-		UE_LOGFMT(LogPixelStreaming2, Log, "[{0}]({1}) = {2}", PlayerId, StatName.ToString(), StatValue);
+		if (!PixelStreaming2Module)
+		{
+			PixelStreaming2Module = FModuleManager::Get().LoadModulePtr<FPixelStreaming2Module>("PixelStreaming2");
+		}
+
+		return PixelStreaming2Module;
 	}
 
-	/**
-	 * IModuleInterface implementation
-	 */
 	void FPixelStreaming2Module::StartupModule()
 	{
-#if UE_SERVER
-		// Hack to no-op the rest of the module so Blueprints can still work
-		return;
-#endif
 		if (!IsStreamingSupported())
 		{
 			return;
@@ -133,34 +54,36 @@ namespace UE::PixelStreaming2
 			return;
 		}
 
-		// Initialize EpicRtc thread. Handles tasks like audio pushing and conference ticking
-		EpicRtcThread = MakeShared<FEpicRtcThread>();
+		// Initialize PixelStreaming thread. Handles tasks like audio pushing and conference ticking
+		PixelStreamingThread = MakeShared<FPixelStreamingThread>();
 
-		// By calling InitDefaultStreamer post engine init we can use pixel streaming in standalone editor mode
-		FCoreDelegates::OnAllModuleLoadingPhasesComplete.AddLambda([this, RHIType]() {
-			// Need to initialize after other modules have initialized such as NVCodec.
-			if (!InitializeEpicRtc())
+		FCoreDelegates::OnAllModuleLoadingPhasesComplete.AddLambda([this]() {
+			if (!IsPlatformSupported())
 			{
 				return;
 			}
 
-			// Check to see if we can use the Pixel Streaming plugin on this platform.
-			// If not then we avoid setting up our delegates to prevent access to the plugin.
-			if (!IsPlatformCompatible())
+			if (!ensure(GEngine != nullptr))
 			{
 				return;
 			}
 
-			if (!ensure(GEngine != nullptr))
+			const EVideoCodec SelectedCodec = UE::PixelStreaming2::GetEnumFromCVar<EVideoCodec>(UPixelStreaming2PluginSettings::CVarEncoderCodec);
+			if ((SelectedCodec == EVideoCodec::H264 && !IsHardwareEncoderSupported<FVideoEncoderConfigH264>())
+				|| (SelectedCodec == EVideoCodec::AV1 && !IsHardwareEncoderSupported<FVideoEncoderConfigAV1>()))
 			{
-				return;
+				UE_LOG(LogPixelStreaming2, Warning, TEXT("Could not setup hardware encoder. This is usually a driver issue or hardware limitation, try reinstalling your drivers."));
+				UE_LOG(LogPixelStreaming2, Warning, TEXT("Falling back to VP8 software video encoding."));
+				UPixelStreaming2PluginSettings::CVarEncoderCodec.AsVariable()->Set(*UE::PixelStreaming2::GetCVarStringFromEnum(EVideoCodec::VP8), ECVF_SetByCommandline);
+				if (UPixelStreaming2Delegates* Delegates = UPixelStreaming2Delegates::Get())
+				{
+					Delegates->OnFallbackToSoftwareEncoding.Broadcast();
+					Delegates->OnFallbackToSoftwareEncodingNative.Broadcast();
+				}
 			}
 
 			FApp::SetUnfocusedVolumeMultiplier(1.0f);
 
-			// Ensure we have ImageWrapper loaded, used in Freezeframes
-			verify(FModuleManager::Get().LoadModule(FName("ImageWrapper")));
-
 			// HACK (Eden.Harris): Until or if we ever find a workaround for fencing, we need to ensure capture always uses a fence.
 			// If we don't then we get frequent and intermittent stuttering as textures are rendered to while being encoded.
 			// From testing NVENC + CUDA pathway seems acceptable without a fence in most cases so we use the faster, unsafer path there.
@@ -174,70 +97,9 @@ namespace UE::PixelStreaming2
 				UPixelStreaming2PluginSettings::CVarCaptureUseFence.AsVariable()->Set(true);
 			}
 
-			// We don't want to start immediately streaming in editor
-			if (!GIsEditor)
-			{
-				InitDefaultStreamer();
-				StartStreaming();
-			}
-
 			bModuleReady = true;
 			ReadyEvent.Broadcast(*this);
 		});
-
-		FModuleManager::LoadModuleChecked<FWebSocketsModule>("WebSockets");
-
-		// Call these to initialize their singletons
-		FStats::Get();
-
-		// Extra Initialisations post loading console commands
-		IConsoleManager::Get().RegisterConsoleCommand(
-			TEXT("PixelStreaming.StartStreaming"),
-			TEXT("Start all streaming sessions"),
-			FConsoleCommandDelegate::CreateLambda([]() {
-				IPixelStreaming2Module::Get().StartStreaming();
-			}));
-
-		IConsoleManager::Get().RegisterConsoleCommand(
-			TEXT("PixelStreaming.StopStreaming"),
-			TEXT("End any existing streaming sessions."),
-			FConsoleCommandDelegate::CreateLambda([]() {
-				IPixelStreaming2Module::Get().StopStreaming();
-			}));
-
-		if (UPixelStreaming2PluginSettings::FDelegates* Delegates = UPixelStreaming2PluginSettings::Delegates())
-		{
-			Delegates->OnLogStatsChanged.AddLambda([this](IConsoleVariable* Var) {
-				bool					   bLogStats = Var->GetBool();
-				UPixelStreaming2Delegates* Delegates = UPixelStreaming2Delegates::Get();
-				if (!Delegates)
-				{
-					return;
-				}
-				if (bLogStats)
-				{
-					LogStatsHandle = Delegates->OnStatChangedNative.AddStatic(&ConsumeStat);
-				}
-				else
-				{
-					Delegates->OnStatChangedNative.Remove(LogStatsHandle);
-				}
-			});
-
-			Delegates->OnWebRTCFpsChanged.AddLambda([](IConsoleVariable*) {
-				IPixelStreaming2Module::Get().ForEachStreamer([](TSharedPtr<IPixelStreaming2Streamer> Streamer) {
-					Streamer->RefreshStreamBitrate();
-				});
-			});
-
-			Delegates->OnWebRTCBitrateChanged.AddLambda([](IConsoleVariable*) {
-				IPixelStreaming2Module::Get().ForEachStreamer([](TSharedPtr<IPixelStreaming2Streamer> Streamer) {
-					Streamer->RefreshStreamBitrate();
-				});
-			});
-		}
-
-		bStartupCompleted = true;
 	}
 
 	void FPixelStreaming2Module::ShutdownModule()
@@ -247,52 +109,15 @@ namespace UE::PixelStreaming2
 			return;
 		}
 
-		if (!bStartupCompleted)
-		{
-			return;
-		}
-
 		// We explicitly call release on streamer so WebRTC gets shutdown before our module is deleted
 		// additionally the streamer does a bunch of delegate calls and unbinds which seem to have issues
 		// when called during engine destruction rather than here.
 		Streamers.Empty();
-		DefaultStreamer.Reset();
 
 		// Reset thread must be called before tasks to ensure it does not attempt to run any partially destroyed tasks from the AudioMixingCapturer
-		EpicRtcThread.Reset();
-		AudioMixingCapturer.Reset();
-
-		TickConferenceTask.Reset();
-
-		if (!EpicRtcPlatform)
-		{
-			UE_LOGFMT(LogPixelStreaming2, Error, "EpicRtcPlatform does not exist during shutdown when it is expected to exist");
-		}
-		else
-		{
-			EpicRtcPlatform->ReleaseConference(ToEpicRtcStringView(EpicRtcConferenceName));
-		}
-
-		bStartupCompleted = false;
-	}
-
-	/**
-	 * End IModuleInterface implementation
-	 */
-
-	FPixelStreaming2Module* FPixelStreaming2Module::GetModule()
-	{
-		if (!PixelStreaming2Module)
-		{
-			PixelStreaming2Module = FModuleManager::Get().LoadModulePtr<FPixelStreaming2Module>("PixelStreaming2");
-		}
-
-		return PixelStreaming2Module;
+		PixelStreamingThread.Reset();
 	}
 
-	/**
-	 * IPixelStreaming2Module implementation
-	 */
 	IPixelStreaming2Module::FReadyEvent& FPixelStreaming2Module::OnReady()
 	{
 		return ReadyEvent;
@@ -305,24 +130,29 @@ namespace UE::PixelStreaming2
 
 	bool FPixelStreaming2Module::StartStreaming()
 	{
-		if (DefaultStreamer.IsValid())
-		{
-			DefaultStreamer->StartStreaming();
-			return true;
-		}
+		bool bStreamingStarted = false;
+		Streamers.Apply([&bStreamingStarted](FString StreamerId, TWeakPtr<IPixelStreaming2Streamer> WeakStreamer) {
+			if (TSharedPtr<IPixelStreaming2Streamer> PinnedStreamer = WeakStreamer.Pin())
+			{
+				PinnedStreamer->StartStreaming();
+				bStreamingStarted = true;
+			}
+		});
 
-		return false;
+		return bStreamingStarted;
 	}
 
 	void FPixelStreaming2Module::StopStreaming()
 	{
-		if (DefaultStreamer.IsValid())
-		{
-			DefaultStreamer->StopStreaming();
-		}
+		Streamers.Apply([](FString StreamerId, TWeakPtr<IPixelStreaming2Streamer> WeakStreamer) {
+			if (TSharedPtr<IPixelStreaming2Streamer> PinnedStreamer = WeakStreamer.Pin())
+			{
+				PinnedStreamer->StopStreaming();
+			}
+		});
 	}
 
-	TSharedPtr<IPixelStreaming2Streamer> FPixelStreaming2Module::CreateStreamer(const FString& StreamerId)
+	TSharedPtr<IPixelStreaming2Streamer> FPixelStreaming2Module::CreateStreamer(const FString& StreamerId, const FString& Type)
 	{
 		TSharedPtr<IPixelStreaming2Streamer> ExistingStreamer = FindStreamer(StreamerId);
 		if (ExistingStreamer)
@@ -330,42 +160,47 @@ namespace UE::PixelStreaming2
 			return ExistingStreamer;
 		}
 
-		TSharedPtr<FStreamer> NewStreamer = FStreamer::Create(StreamerId, EpicRtcConference);
-		{
-			FScopeLock Lock(&StreamersCS);
-			Streamers.Add(StreamerId, NewStreamer);
-		}
-
+		IPixelStreaming2StreamerFactory*	 Factory = IPixelStreaming2StreamerFactory::Get(Type);
+		TSharedPtr<IPixelStreaming2Streamer> NewStreamer = Factory->CreateNewStreamer(StreamerId);
+		NewStreamer->Initialize();
 		// Any time we create a new streamer, populate it's signalling server URL with whatever is in the ini, console or command line
-		NewStreamer->SetSignallingServerURL(UPixelStreaming2PluginSettings::CVarSignallingURL.GetValueOnAnyThread());
+		NewStreamer->SetConnectionURL(UPixelStreaming2PluginSettings::CVarConnectionURL.GetValueOnAnyThread());
 
-		// Ensure that this new streamer is able to handle pixel streaming relevant input
 		RegisterCustomHandlers(NewStreamer);
 
+		Streamers.Add(StreamerId, NewStreamer);
+
 		return NewStreamer;
 	}
 
-	TSharedPtr<IPixelStreaming2AudioProducer> FPixelStreaming2Module::CreateAudioProducer()
+	TSharedPtr<IPixelStreaming2VideoProducer> FPixelStreaming2Module::CreateVideoProducer()
 	{
-		return GetAudioCapturer()->CreateAudioProducer();
+		return FVideoProducer::Create();
 	}
 
-	TSharedPtr<IPixelStreaming2VideoProducer> FPixelStreaming2Module::CreateVideoProducer()
+	TSharedPtr<IPixelStreaming2AudioProducer> FPixelStreaming2Module::CreateAudioProducer()
 	{
-		return FVideoProducer::Create();
+		// TODO (william.belcher): This will be fixed with the refactor to audio mixing.
+		return nullptr;
 	}
 
 	TArray<FString> FPixelStreaming2Module::GetStreamerIds()
 	{
-		TArray<FString> StreamerKeys;
-		FScopeLock		Lock(&StreamersCS);
-		Streamers.GenerateKeyArray(StreamerKeys);
-		return StreamerKeys;
+		TSet<FString> KeySet;
+		{
+			Streamers.GetKeys(KeySet);
+		}
+
+		TArray<FString> KeyArray;
+		for (const FString& Key : KeySet)
+		{
+			KeyArray.Add(Key);
+		}
+		return KeyArray;
 	}
 
 	TSharedPtr<IPixelStreaming2Streamer> FPixelStreaming2Module::FindStreamer(const FString& StreamerId)
 	{
-		FScopeLock Lock(&StreamersCS);
 		if (Streamers.Contains(StreamerId))
 		{
 			return Streamers[StreamerId].Pin();
@@ -376,7 +211,6 @@ namespace UE::PixelStreaming2
 	TSharedPtr<IPixelStreaming2Streamer> FPixelStreaming2Module::DeleteStreamer(const FString& StreamerId)
 	{
 		TSharedPtr<IPixelStreaming2Streamer> ToBeDeleted;
-		FScopeLock							 Lock(&StreamersCS);
 		if (Streamers.Contains(StreamerId))
 		{
 			ToBeDeleted = Streamers[StreamerId].Pin();
@@ -387,49 +221,36 @@ namespace UE::PixelStreaming2
 
 	void FPixelStreaming2Module::DeleteStreamer(TSharedPtr<IPixelStreaming2Streamer> ToBeDeleted)
 	{
-		FScopeLock Lock(&StreamersCS);
-		for (auto& [Id, Streamer] : Streamers)
-		{
-			if (Streamer == ToBeDeleted)
+		FString StreamerIdToRemove;
+		Streamers.ApplyUntil([&StreamerIdToRemove, &ToBeDeleted](FString StreamerId, TWeakPtr<IPixelStreaming2Streamer> WeakStreamer) {
+			if (WeakStreamer.Pin() == ToBeDeleted)
 			{
-				Streamers.Remove(Id);
-				break;
+				StreamerIdToRemove = StreamerId;
+				return true;
 			}
+			return false;
+		});
+		if (!StreamerIdToRemove.IsEmpty())
+		{
+			Streamers.Remove(StreamerIdToRemove);
 		}
 	}
 
-	void FPixelStreaming2Module::AddInputComponent(UPixelStreaming2Input* InInputComponent)
-	{
-		InputComponents.Add(InInputComponent);
-	}
-
-	void FPixelStreaming2Module::RemoveInputComponent(UPixelStreaming2Input* InInputComponent)
-	{
-		InputComponents.Remove(InInputComponent);
-	}
-
-	const TArray<UPixelStreaming2Input*> FPixelStreaming2Module::GetInputComponents()
-	{
-		return InputComponents;
-	}
-
 	FString FPixelStreaming2Module::GetDefaultStreamerID()
 	{
 		return UPixelStreaming2PluginSettings::CVarDefaultStreamerID.GetValueOnAnyThread();
 	}
 
-	FString FPixelStreaming2Module::GetDefaultSignallingURL()
+	FString FPixelStreaming2Module::GetDefaultConnectionURL()
 	{
-		return UPixelStreaming2PluginSettings::CVarSignallingURL.GetValueOnAnyThread();
+		return UPixelStreaming2PluginSettings::CVarConnectionURL.GetValueOnAnyThread();
 	}
 
 	void FPixelStreaming2Module::ForEachStreamer(const TFunction<void(TSharedPtr<IPixelStreaming2Streamer>)>& Func)
 	{
 		TSet<FString> KeySet;
-		{
-			FScopeLock Lock(&StreamersCS);
-			Streamers.GetKeys(KeySet);
-		}
+		Streamers.GetKeys(KeySet);
+
 		for (auto&& StreamerId : KeySet)
 		{
 			if (TSharedPtr<IPixelStreaming2Streamer> Streamer = FindStreamer(StreamerId))
@@ -439,91 +260,19 @@ namespace UE::PixelStreaming2
 		}
 	}
 
-	/**
-	 * End IPixelStreaming2Module implementation
-	 */
-
-	TSharedPtr<FEpicRtcAudioMixingCapturer> FPixelStreaming2Module::GetAudioCapturer()
+	void FPixelStreaming2Module::AddInputComponent(UPixelStreaming2Input* InInputComponent)
 	{
-		if (!AudioMixingCapturer)
-		{
-			AudioMixingCapturer = FEpicRtcAudioMixingCapturer::Create();
-		}
-
-		return AudioMixingCapturer;
+		InputComponents.Add(InInputComponent);
 	}
 
-	void FPixelStreaming2Module::InitDefaultStreamer()
+	void FPixelStreaming2Module::RemoveInputComponent(UPixelStreaming2Input* InInputComponent)
 	{
-		UE_LOG(LogPixelStreaming2, Log, TEXT("PixelStreaming2 streamer ID: %s"), *GetDefaultStreamerID());
-
-		DefaultStreamer = CreateStreamer(GetDefaultStreamerID());
-		// The PixelStreaming2EditorModule handles setting video input in the editor
-		if (!GIsEditor)
-		{
-			// default to the scene viewport if we have a game engine
-			if (UGameEngine* GameEngine = Cast<UGameEngine>(GEngine))
-			{
-				TSharedPtr<SWindow>						 TargetWindow = GameEngine->GameViewport->GetWindow();
-				TSharedPtr<IPixelStreaming2InputHandler> InputHandler = DefaultStreamer->GetInputHandler().Pin();
-				if (TargetWindow.IsValid() && InputHandler.IsValid())
-				{
-					InputHandler->SetTargetWindow(TargetWindow);
-				}
-				else
-				{
-					UE_LOG(LogPixelStreaming2, Error, TEXT("Cannot set target window - target window is not valid."));
-				}
-			}
-		}
-
-		if (!DefaultStreamer->GetSignallingServerURL().IsEmpty())
-		{
-			// The user has specified a URL on the command line meaning their intention is to start streaming immediately
-			// in that case, set up the video input for them (as long as we're not in editor)
-			if (UPixelStreaming2PluginSettings::CVarUseMediaCapture.GetValueOnAnyThread())
-			{
-				DefaultStreamer->SetVideoProducer(FVideoProducerMediaCapture::CreateActiveViewportCapture());
-			}
-			else
-			{
-				DefaultStreamer->SetVideoProducer(FVideoProducerBackBuffer::Create());
-			}
-		}
+		InputComponents.Remove(InInputComponent);
 	}
 
-	bool FPixelStreaming2Module::IsPlatformCompatible() const
+	const TArray<UPixelStreaming2Input*> FPixelStreaming2Module::GetInputComponents()
 	{
-		bool bCompatible = true;
-
-#if PLATFORM_WINDOWS
-		bool bWin8OrHigher = IsWindows8OrGreater();
-		if (!bWin8OrHigher)
-		{
-			FString ErrorString(TEXT("Failed to initialize Pixel Streaming plugin because minimum requirement is Windows 8"));
-			FText	ErrorText = FText::FromString(ErrorString);
-			FText	TitleText = FText::FromString(TEXT("Pixel Streaming Plugin"));
-			FMessageDialog::Open(EAppMsgType::Ok, ErrorText, TitleText);
-			UE_LOG(LogPixelStreaming2, Error, TEXT("%s"), *ErrorString);
-			bCompatible = false;
-		}
-#endif
-
-		const EVideoCodec SelectedCodec = GetEnumFromCVar<EVideoCodec>(UPixelStreaming2PluginSettings::CVarEncoderCodec);
-		if ((SelectedCodec == EVideoCodec::H264 && !IsHardwareEncoderSupported<FVideoEncoderConfigH264>())
-			|| (SelectedCodec == EVideoCodec::AV1 && !IsHardwareEncoderSupported<FVideoEncoderConfigAV1>()))
-		{
-			UE_LOG(LogPixelStreaming2, Warning, TEXT("Could not setup hardware encoder. This is usually a driver issue or hardware limitation, try reinstalling your drivers."));
-			UE_LOG(LogPixelStreaming2, Warning, TEXT("Falling back to VP8 software video encoding."));
-			UPixelStreaming2PluginSettings::CVarEncoderCodec.AsVariable()->Set(*GetCVarStringFromEnum(EVideoCodec::VP8), ECVF_SetByCommandline);
-			if (UPixelStreaming2Delegates* Delegates = UPixelStreaming2Delegates::Get())
-			{
-				Delegates->OnFallbackToSoftwareEncoding.Broadcast();
-				Delegates->OnFallbackToSoftwareEncodingNative.Broadcast();
-			}
-		}
-
-		return bCompatible;
+		return InputComponents;
 	}
 
 	void FPixelStreaming2Module::RegisterCustomHandlers(TSharedPtr<IPixelStreaming2Streamer> Streamer)
@@ -611,131 +360,6 @@ namespace UE::PixelStreaming2
 		}
 	}
 
-	FString FPixelStreaming2Module::GetFieldTrials()
-	{
-		FString FieldTrials = UPixelStreaming2PluginSettings::CVarWebRTCFieldTrials.GetValueOnAnyThread();
-
-		// Set the WebRTC-FrameDropper/Disabled/ if the CVar is set
-		if (UPixelStreaming2PluginSettings::CVarWebRTCDisableFrameDropper.GetValueOnAnyThread())
-		{
-			FieldTrials += TEXT("WebRTC-FrameDropper/Disabled/");
-		}
-
-		if (UPixelStreaming2PluginSettings::CVarWebRTCEnableFlexFec.GetValueOnAnyThread())
-		{
-			FieldTrials += TEXT("WebRTC-FlexFEC-03-Advertised/Enabled/WebRTC-FlexFEC-03/Enabled/");
-		}
-
-		// Parse "WebRTC-Video-Pacing/" field trial
-		{
-			float PacingFactor = UPixelStreaming2PluginSettings::CVarWebRTCVideoPacingFactor.GetValueOnAnyThread();
-			float PacingMaxDelayMs = UPixelStreaming2PluginSettings::CVarWebRTCVideoPacingMaxDelay.GetValueOnAnyThread();
-			;
-			if (PacingFactor >= 0.0f || PacingMaxDelayMs >= 0.0f)
-			{
-				FString VideoPacingFieldTrialStr = TEXT("WebRTC-Video-Pacing/");
-				bool	bHasPacingFactor = PacingFactor >= 0.0f;
-				if (bHasPacingFactor)
-				{
-					VideoPacingFieldTrialStr += FString::Printf(TEXT("factor:%.1f"), PacingFactor);
-				}
-				bool bHasMaxDelay = PacingMaxDelayMs >= 0.0f;
-				if (bHasMaxDelay)
-				{
-					VideoPacingFieldTrialStr += bHasPacingFactor ? TEXT(",") : TEXT("");
-					VideoPacingFieldTrialStr += FString::Printf(TEXT("max_delay:%.0f"), PacingMaxDelayMs);
-				}
-				VideoPacingFieldTrialStr += TEXT("/");
-				FieldTrials += VideoPacingFieldTrialStr;
-			}
-		}
-
-		return FieldTrials;
-	}
-
-	bool FPixelStreaming2Module::InitializeEpicRtc()
-	{
-		// Initialize the audio capturer (if not already called by some user code)
-		AudioMixingCapturer = GetAudioCapturer();
-
-		EpicRtcVideoEncoderInitializers = { new FEpicRtcVideoEncoderInitializer() };
-		EpicRtcVideoDecoderInitializers = { new FEpicRtcVideoDecoderInitializer() };
-
-		EpicRtcPlatformConfig PlatformConfig{
-			._memory = new FEpicRtcAllocator()
-		};
-
-		EpicRtcErrorCode Result = GetOrCreatePlatform(PlatformConfig, EpicRtcPlatform.GetInitReference());
-		if (Result != EpicRtcErrorCode::Ok && Result != EpicRtcErrorCode::FoundExistingPlatform)
-		{
-			UE_LOG(LogPixelStreaming2, Warning, TEXT("Unable to create EpicRtc Platform. GetOrCreatePlatform returned %s"), *ToString(Result));
-			return false;
-		}
-
-		FUtf8String EpicRtcFieldTrials(GetFieldTrials());
-
-		WebsocketFactory = MakeRefCount<FEpicRtcWebsocketFactory>();
-
-		StatsCollector = MakeRefCount<FEpicRtcStatsCollector>();
-
-		// clang-format off
-		EpicRtcConfig ConferenceConfig = {
-			._websocketFactory = WebsocketFactory.GetReference(),
-			._signallingType = EpicRtcSignallingType::PixelStreaming,
-			._signingPlugin = nullptr,
-			._migrationPlugin = nullptr,
-			._audioDevicePlugin = nullptr,
-			._audioConfig = {
-				._tickAdm = true,
-				._audioEncoderInitializers = {}, // Not needed because we use the inbuilt audio codecs
-				._audioDecoderInitializers = {}, // Not needed because we use the inbuilt audio codecs
-				._enableBuiltInAudioCodecs = true,
-			},
-			._videoConfig = {
-				._videoEncoderInitializers = {
-					._ptr = const_cast<const EpicRtcVideoEncoderInitializerInterface**>(EpicRtcVideoEncoderInitializers.GetData()),
-					._size = (uint64_t)EpicRtcVideoEncoderInitializers.Num()
-				},
-				._videoDecoderInitializers = {
-					._ptr = const_cast<const EpicRtcVideoDecoderInitializerInterface**>(EpicRtcVideoDecoderInitializers.GetData()),
-					._size = (uint64_t)EpicRtcVideoDecoderInitializers.Num()
-				},
-				._enableBuiltInVideoCodecs = false
-			},
-			._fieldTrials = {
-				._fieldTrials = ToEpicRtcStringView(EpicRtcFieldTrials),
-				._isGlobal = 0
-			},
-			._logging = {
-				._logger = new FEpicRtcLogsRedirector(),
-#if !NO_LOGGING // When building WITH_SHIPPING by default .GetVerbosity() does not exist
-				._level = UnrealLogToEpicRtcCategoryMap[LogPixelStreaming2EpicRtc.GetVerbosity()],
-				._levelWebRtc = UnrealLogToEpicRtcCategoryMap[LogPixelStreaming2WebRtc.GetVerbosity()]
-#endif
-			},
-			._stats = {
-				._statsCollectorCallback = StatsCollector.GetReference(),
-				._statsCollectorInterval = 1000,
-				._jsonFormatOnly = false
-			}
-		};
-		// clang-format on
-
-		Result = EpicRtcPlatform->CreateConference(ToEpicRtcStringView(EpicRtcConferenceName), ConferenceConfig, EpicRtcConference.GetInitReference());
-		if (Result != EpicRtcErrorCode::Ok)
-		{
-			UE_LOG(LogPixelStreaming2, Warning, TEXT("Unable to create EpicRtc Conference: CreateConference returned %s"), *ToString(Result));
-			return false;
-		}
-
-		TickConferenceTask = FEpicRtcTickableTask::Create<FEpicRtcTickConferenceTask>(EpicRtcConference, TEXT("PixelStreaming2Module TickConferenceTask"));
-
-		return true;
-	}
-
-	/**
-	 * End own methods
-	 */
 } // namespace UE::PixelStreaming2
 
 IMPLEMENT_MODULE(UE::PixelStreaming2::FPixelStreaming2Module, PixelStreaming2)
diff --git a/Engine/Plugins/Media/PixelStreaming2/Source/PixelStreaming2/Private/PixelStreaming2Module.h b/Engine/Plugins/Media/PixelStreaming2/Source/PixelStreaming2/Private/PixelStreaming2Module.h
--- a/Engine/Plugins/Media/PixelStreaming2/Source/PixelStreaming2/Private/PixelStreaming2Module.h
+++ b/Engine/Plugins/Media/PixelStreaming2/Source/PixelStreaming2/Private/PixelStreaming2Module.h
@@ -3,22 +3,12 @@
 #pragma once
 
 #include "IPixelStreaming2Module.h"
-#include "EpicRtcConferenceUtils.h"
-#include "EpicRtcStatsCollector.h"
-
-#include "epic_rtc/core/platform.h"
-#include "epic_rtc/plugins/signalling/signalling_type.h"
+#include "ThreadSafeMap.h"
 
 class UPixelStreaming2Input;
-class SWindow;
 
 namespace UE::PixelStreaming2
 {
-	class FStreamer;
-	class FVideoInputBackBuffer;
-	class FVideoSourceGroup;
-	class FEpicRtcWebsocketFactory;
-
 	/**
 	 * This plugin allows the back buffer to be sent as a compressed video across a network.
 	 */
@@ -29,31 +19,25 @@ namespace UE::PixelStreaming2
 
 		virtual ~FPixelStreaming2Module() = default;
 
-		// Begin IPixelStreaming2Module
 		virtual FReadyEvent&							  OnReady() override;
 		virtual bool									  IsReady() override;
 		virtual bool									  StartStreaming() override;
 		virtual void									  StopStreaming() override;
-		virtual TSharedPtr<IPixelStreaming2Streamer>	  CreateStreamer(const FString& StreamerId) override;
-		virtual TSharedPtr<IPixelStreaming2AudioProducer> CreateAudioProducer() override;
+		virtual TSharedPtr<IPixelStreaming2Streamer>	  CreateStreamer(const FString& StreamerId, const FString& Type = TEXT("DefaultRtc")) override;
 		virtual TSharedPtr<IPixelStreaming2VideoProducer> CreateVideoProducer() override;
+		virtual TSharedPtr<IPixelStreaming2AudioProducer> CreateAudioProducer() override;
 		virtual TArray<FString>							  GetStreamerIds() override;
 		virtual TSharedPtr<IPixelStreaming2Streamer>	  FindStreamer(const FString& StreamerId) override;
 		virtual TSharedPtr<IPixelStreaming2Streamer>	  DeleteStreamer(const FString& StreamerId) override;
-		void											  DeleteStreamer(TSharedPtr<IPixelStreaming2Streamer> ToBeDeleted) override;
+		virtual void									  DeleteStreamer(TSharedPtr<IPixelStreaming2Streamer> ToBeDeleted) override;
 		virtual FString									  GetDefaultStreamerID() override;
-		virtual FString									  GetDefaultSignallingURL() override;
+		virtual FString									  GetDefaultConnectionURL() override;
+		virtual void									  ForEachStreamer(const TFunction<void(TSharedPtr<IPixelStreaming2Streamer>)>& Func) override;
 
 		// These are staying on the module at the moment as theres no way of the BPs knowing which streamer they are relevant to
-		virtual void								 AddInputComponent(UPixelStreaming2Input* InInputComponent);
-		virtual void								 RemoveInputComponent(UPixelStreaming2Input* InInputComponent);
-		virtual const TArray<UPixelStreaming2Input*> GetInputComponents();
-		virtual void								 ForEachStreamer(const TFunction<void(TSharedPtr<IPixelStreaming2Streamer>)>& Func) override;
-		// End IPixelStreaming2Module
-
-		TSharedPtr<class FEpicRtcAudioMixingCapturer> GetAudioCapturer();
-		TRefCountPtr<EpicRtcConferenceInterface>	  GetEpicRtcConference() { return EpicRtcConference; }
-		TRefCountPtr<FEpicRtcStatsCollector>		  GetStatsCollector() { return StatsCollector; }
+		void								 AddInputComponent(UPixelStreaming2Input* InInputComponent);
+		void								 RemoveInputComponent(UPixelStreaming2Input* InInputComponent);
+		const TArray<UPixelStreaming2Input*> GetInputComponents();
 
 	private:
 		// Begin IModuleInterface
@@ -61,42 +45,21 @@ namespace UE::PixelStreaming2
 		void ShutdownModule() override;
 		// End IModuleInterface
 
-		// Own methods
-		void InitDefaultStreamer();
-		bool IsPlatformCompatible() const;
+	private:
 		void RegisterCustomHandlers(TSharedPtr<IPixelStreaming2Streamer> Streamer);
 		void HandleUIInteraction(FMemoryReader Ar);
 
-		FString GetFieldTrials();
-		bool	InitializeEpicRtc();
-
 	private:
-		bool						   bModuleReady = false;
-		bool						   bStartupCompleted = false;
-		static FPixelStreaming2Module* PixelStreaming2Module;
-
-		FReadyEvent		ReadyEvent;
-		FDelegateHandle LogStatsHandle;
+		bool		bModuleReady = false;
+		FReadyEvent ReadyEvent;
 
-		TArray<UPixelStreaming2Input*>					  InputComponents;
-		mutable FCriticalSection						  StreamersCS;
-		TMap<FString, TWeakPtr<IPixelStreaming2Streamer>> Streamers;
-		TSharedPtr<IPixelStreaming2Streamer>			  DefaultStreamer;
+		TArray<UPixelStreaming2Input*> InputComponents;
 
-	private:
-		// FEpicRtcThread must exist before any AudioTask and AudioMixingCapturer(Which contains a audio task) to ensure it is destroyed last
-		TSharedPtr<class FEpicRtcThread>			  EpicRtcThread;
-		TSharedPtr<class FEpicRtcAudioMixingCapturer> AudioMixingCapturer;
-		TRefCountPtr<EpicRtcPlatformInterface>		  EpicRtcPlatform;
-		TRefCountPtr<EpicRtcConferenceInterface>	  EpicRtcConference;
-		TRefCountPtr<FEpicRtcStatsCollector>		  StatsCollector;
-
-		TRefCountPtr<FEpicRtcWebsocketFactory>	   WebsocketFactory;
-		TUniqueTaskPtr<FEpicRtcTickConferenceTask> TickConferenceTask;
+		// FPixelStreamingThread must exist before any AudioTask and AudioMixingCapturer(Which contains a audio task) to ensure it is destroyed last
+		TSharedPtr<class FPixelStreamingThread> PixelStreamingThread;
 
-		TArray<EpicRtcVideoEncoderInitializerInterface*> EpicRtcVideoEncoderInitializers;
-		TArray<EpicRtcVideoDecoderInitializerInterface*> EpicRtcVideoDecoderInitializers;
+		TThreadSafeMap<FString, TWeakPtr<IPixelStreaming2Streamer>> Streamers;
 
-		static FUtf8String EpicRtcConferenceName;
+		static FPixelStreaming2Module* PixelStreaming2Module;
 	};
-} // namespace UE::PixelStreaming2
+} // namespace UE::PixelStreaming2
\ No newline at end of file
