[PS2] PR #12718: Refactor PixelStreaming2 plugin to enable other streaming protocols

diff --git a/Engine/Plugins/Media/PixelStreaming2/Source/PixelStreaming2/Private/Streamer.cpp b/Engine/Plugins/Media/PixelStreaming2/Source/PixelStreaming2/Private/Streamer.cpp
--- a/Engine/Plugins/Media/PixelStreaming2/Source/PixelStreaming2/Private/Streamer.cpp
+++ b/Engine/Plugins/Media/PixelStreaming2/Source/PixelStreaming2/Private/Streamer.cpp
@@ -2,1632 +2,126 @@
 
 #include "Streamer.h"
 
-#include "IPixelStreaming2Module.h"
-#include "PixelStreaming2PluginSettings.h"
-#include "PixelStreaming2Delegates.h"
-#include "EpicRtcWebsocket.h"
-#include "EpicRtcWebsocketFactory.h"
-#include "TextureResource.h"
-#include "WebSocketsModule.h"
+#include "Features/IModularFeatures.h"
 #include "Logging.h"
-#include "Stats.h"
-#include "PixelStreaming2StatNames.h"
-#include "Async/Async.h"
-#include "Engine/Texture2D.h"
-#include "RTCStatsCollector.h"
-#include "Framework/Application/SlateApplication.h"
-#include "Dom/JsonValue.h"
-#include "Serialization/JsonWriter.h"
-#include "Serialization/JsonSerializer.h"
-#include "Dom/JsonObject.h"
-#include "PixelStreaming2Module.h"
-#include "EpicRtcAudioTrack.h"
-#include "EpicRtcDataTrack.h"
-#include "IPixelStreaming2Module.h"
-#include "IPixelStreaming2InputModule.h"
-#include "CoreGlobals.h"
-#include "UtilsString.h"
-#include "UtilsVideo.h"
+
+static FName ModularFeatureName = TEXT("PixelStreaming2 Streamer");
 
 namespace UE::PixelStreaming2
 {
-	TSharedPtr<FStreamer> FStreamer::Create(const FString& StreamerId, TRefCountPtr<EpicRtcConferenceInterface> Conference)
-	{
-		TSharedPtr<FStreamer> Streamer = TSharedPtr<FStreamer>(new FStreamer(StreamerId, Conference));
-
-		if (TSharedPtr<IPixelStreaming2InputHandler> InputHandler = Streamer->GetInputHandler().Pin())
-		{
-			InputHandler->GetToStreamerProtocol()->OnProtocolUpdated().AddSP(Streamer.ToSharedRef(), &FStreamer::OnProtocolUpdated);
-			InputHandler->GetFromStreamerProtocol()->OnProtocolUpdated().AddSP(Streamer.ToSharedRef(), &FStreamer::OnProtocolUpdated);
-		}
-
-		Streamer->EpicRtcManager->OnSessionStateUpdate.AddSP(Streamer.ToSharedRef(), &FStreamer::OnSessionStateUpdate);
-		Streamer->EpicRtcManager->OnSessionErrorUpdate.AddSP(Streamer.ToSharedRef(), &FStreamer::OnSessionErrorUpdate);
-		Streamer->EpicRtcManager->OnSessionRoomsAvailableUpdate.AddSP(Streamer.ToSharedRef(), &FStreamer::OnSessionRoomsAvailableUpdate);
-
-		Streamer->EpicRtcManager->OnRoomStateUpdate.AddSP(Streamer.ToSharedRef(), &FStreamer::OnRoomStateUpdate);
-		Streamer->EpicRtcManager->OnRoomJoinedUpdate.AddSP(Streamer.ToSharedRef(), &FStreamer::OnRoomJoinedUpdate);
-		Streamer->EpicRtcManager->OnRoomLeftUpdate.AddSP(Streamer.ToSharedRef(), &FStreamer::OnRoomLeftUpdate);
-		Streamer->EpicRtcManager->OnAudioTrackUpdate.AddSP(Streamer.ToSharedRef(), &FStreamer::OnAudioTrackUpdate);
-		Streamer->EpicRtcManager->OnVideoTrackUpdate.AddSP(Streamer.ToSharedRef(), &FStreamer::OnVideoTrackUpdate);
-		Streamer->EpicRtcManager->OnDataTrackUpdate.AddSP(Streamer.ToSharedRef(), &FStreamer::OnDataTrackUpdate);
-		Streamer->EpicRtcManager->OnLocalSdpUpdate.AddSP(Streamer.ToSharedRef(), &FStreamer::OnLocalSdpUpdate);
-		Streamer->EpicRtcManager->OnRemoteSdpUpdate.AddSP(Streamer.ToSharedRef(), &FStreamer::OnRemoteSdpUpdate);
-		Streamer->EpicRtcManager->OnRoomErrorUpdate.AddSP(Streamer.ToSharedRef(), &FStreamer::OnRoomErrorUpdate);
-
-		Streamer->EpicRtcManager->OnAudioTrackMuted.AddSP(Streamer.ToSharedRef(), &FStreamer::OnAudioTrackMuted);
-		Streamer->EpicRtcManager->OnAudioTrackFrame.AddSP(Streamer.ToSharedRef(), &FStreamer::OnAudioTrackFrame);
-		Streamer->EpicRtcManager->OnAudioTrackRemoved.AddSP(Streamer.ToSharedRef(), &FStreamer::OnAudioTrackRemoved);
-		Streamer->EpicRtcManager->OnAudioTrackState.AddSP(Streamer.ToSharedRef(), &FStreamer::OnAudioTrackState);
-
-		Streamer->EpicRtcManager->OnVideoTrackMuted.AddSP(Streamer.ToSharedRef(), &FStreamer::OnVideoTrackMuted);
-		Streamer->EpicRtcManager->OnVideoTrackFrame.AddSP(Streamer.ToSharedRef(), &FStreamer::OnVideoTrackFrame);
-		Streamer->EpicRtcManager->OnVideoTrackRemoved.AddSP(Streamer.ToSharedRef(), &FStreamer::OnVideoTrackRemoved);
-		Streamer->EpicRtcManager->OnVideoTrackState.AddSP(Streamer.ToSharedRef(), &FStreamer::OnVideoTrackState);
-
-		Streamer->EpicRtcManager->OnDataTrackRemoved.AddSP(Streamer.ToSharedRef(), &FStreamer::OnDataTrackRemoved);
-		Streamer->EpicRtcManager->OnDataTrackState.AddSP(Streamer.ToSharedRef(), &FStreamer::OnDataTrackState);
-		Streamer->EpicRtcManager->OnDataTrackMessage.AddSP(Streamer.ToSharedRef(), &FStreamer::OnDataTrackMessage);
-
-		FPixelStreaming2Module::GetModule()->GetStatsCollector()->OnStatsReady.AddSP(Streamer.ToSharedRef(), &FStreamer::OnStatsReady);
-
-		return Streamer;
-	}
-
-	FStreamer::FStreamer(const FString& InStreamerId, TRefCountPtr<EpicRtcConferenceInterface> Conference)
-		: StreamerId(InStreamerId)
-		, InputHandler(IPixelStreaming2InputModule::Get().CreateInputHandler())
-		, Players(new TThreadSafeMap<FString, FPlayerContext>())
-		, VideoCapturer(FVideoCapturer::Create())
-		, VideoSourceGroup(FVideoSourceGroup::Create(VideoCapturer))
-		, FreezeFrame(FFreezeFrame::Create(Players, VideoCapturer, InputHandler))
-		, EpicRtcManager(MakeShared<FEpicRtcManager>())
-	{
-		ReconnectTimer = MakeShared<FStreamerReconnectTimer>();
-
-		InputHandler->SetElevatedCheck([this](FString PlayerId) {
-			return GetEnumFromCVar<EInputControllerMode>(UPixelStreaming2PluginSettings::CVarInputController) == EInputControllerMode::Any
-				|| InputControllingId == INVALID_PLAYER_ID
-				|| PlayerId == InputControllingId;
-		});
-
-		EpicRtcManager->EpicRtcConference = Conference;
-
-		EpicRtcManager->SessionObserver = MakeRefCount<FEpicRtcSessionObserver>(EpicRtcManager);
-		EpicRtcManager->RoomObserver = MakeRefCount<FEpicRtcRoomObserver>(EpicRtcManager);
-
-		EpicRtcManager->AudioTrackObserverFactory = MakeRefCount<FEpicRtcAudioTrackObserverFactory>(EpicRtcManager);
-		EpicRtcManager->VideoTrackObserverFactory = MakeRefCount<FEpicRtcVideoTrackObserverFactory>(EpicRtcManager);
-		EpicRtcManager->DataTrackObserverFactory = MakeRefCount<FEpicRtcDataTrackObserverFactory>(EpicRtcManager);
-	}
-
-	FStreamer::~FStreamer()
-	{
-		StopStreaming();
-	}
-
-	void FStreamer::OnProtocolUpdated()
-	{
-		Players->Apply([this](FString DataPlayerId, FPlayerContext& PlayerContext) {
-			SendProtocol(DataPlayerId);
-		});
-	}
-
-	void FStreamer::SetStreamFPS(int32 InFramesPerSecond)
-	{
-		VideoSourceGroup->SetFPS(InFramesPerSecond);
-	}
-
-	int32 FStreamer::GetStreamFPS()
-	{
-		return VideoSourceGroup->GetFPS();
-	}
-
-	void FStreamer::SetCoupleFramerate(bool bCouple)
-	{
-		VideoSourceGroup->SetDecoupleFramerate(!bCouple);
-	}
-
-	void FStreamer::SetVideoProducer(TSharedPtr<IPixelStreaming2VideoProducer> Producer)
-	{
-		VideoCapturer->SetVideoProducer(StaticCastSharedPtr<FVideoProducer>(Producer));
-	}
-
-	TWeakPtr<IPixelStreaming2VideoProducer> FStreamer::GetVideoProducer()
-	{
-		return VideoCapturer->GetVideoProducer();
-	}
-
-	void FStreamer::SetSignallingServerURL(const FString& InSignallingServerURL)
-	{
-		CurrentSignallingServerURL = InSignallingServerURL;
-	}
-
-	FString FStreamer::GetSignallingServerURL()
-	{
-		return CurrentSignallingServerURL;
-	}
-
-	void FStreamer::StartStreaming()
-	{
-		if (CurrentSignallingServerURL.IsEmpty())
-		{
-			UE_LOG(LogPixelStreaming2, Log, TEXT("Attempted to start streamer (%s) but no signalling server URL has been set. Use Streamer->SetSignallingServerURL(URL) or -PixelStreaming2URL="), *StreamerId);
-			return;
-		}
-
-		StopStreaming();
-		ReconnectTimer->Stop();
-
-		if (UPixelStreaming2Delegates* Delegates = UPixelStreaming2Delegates::Get())
-		{
-			ConsumeStatsHandle = Delegates->OnStatChangedNative.AddSP(AsShared(), &FStreamer::ConsumeStats);
-			AllConnectionsClosedHandle = Delegates->OnAllConnectionsClosedNative.AddSP(AsShared(), &FStreamer::TriggerMouseLeave);
-		}
-
-		VideoCapturer->ResetFrameCapturer();
-
-		// Broadcast the preconnection event just before we do `TryConnect`
-		StreamingPreConnectionEvent.Broadcast(this);
-
-		VideoSourceGroup->Start();
-
-		FUtf8String Utf8StreamerId(StreamerId);
-		FUtf8String Utf8CurrentSignallingServerURL(CurrentSignallingServerURL);
-
-		EpicRtcSessionConfig SessionConfig{
-			._id = ToEpicRtcStringView(Utf8StreamerId),
-			._url = ToEpicRtcStringView(Utf8CurrentSignallingServerURL),
-			._observer = EpicRtcManager->SessionObserver
-		};
-
-		EpicRtcErrorCode Result = EpicRtcManager->EpicRtcConference->CreateSession(SessionConfig, EpicRtcManager->EpicRtcSession.GetInitReference());
-		if (Result != EpicRtcErrorCode::Ok)
-		{
-			UE_LOG(LogPixelStreaming2, Error, TEXT("Failed to create EpicRtc session. CreateSession returned %s"), *ToString(Result));
-			StopStreaming();
-			return;
-		}
-
-		Result = EpicRtcManager->EpicRtcSession->Connect();
-		if (Result != EpicRtcErrorCode::Ok)
-		{
-			UE_LOG(LogPixelStreaming2, Error, TEXT("Failed to connect EpicRtcSession. Connect returned %s"), *ToString(Result));
-			StopStreaming();
-			return;
-		}
-
-		bStreamingStarted = true;
-	}
-
-	void FStreamer::StopStreaming()
-	{
-		if (UPixelStreaming2Delegates* Delegates = UPixelStreaming2Delegates::Get())
-		{
-			Delegates->OnStatChangedNative.Remove(ConsumeStatsHandle);
-			Delegates->OnAllConnectionsClosedNative.Remove(AllConnectionsClosedHandle);
-		}
-
-		RemoveSession(true);
-
-		VideoSourceGroup->Stop();
-		TriggerMouseLeave(StreamerId);
-
-		if (bStreamingStarted)
-		{
-			OnStreamingStopped().Broadcast(this);
-		}
-
-		DeleteAllPlayerSessions();
-		bStreamingStarted = false;
-	}
-
-	void FStreamer::RemoveSession(bool bDisconnect)
-	{
-		if (!EpicRtcManager->EpicRtcSession)
-		{
-			return;
-		}
-
-		RemoveRoom();
-
-		if (bDisconnect)
-		{
-			EpicRtcErrorCode Result = EpicRtcManager->EpicRtcSession->Disconnect(ToEpicRtcStringView("Streaming Session Removed"));
-			if (Result == EpicRtcErrorCode::SessionDisconnected)
-			{
-				UE_LOG(LogPixelStreaming2, VeryVerbose, TEXT("FStreamer::StopStreaming - Session disconnected cleanly."));
-			}
-			else if (Result != EpicRtcErrorCode::Ok)
-			{
-				UE_LOG(LogPixelStreaming2, Error, TEXT("Failed to disconnect EpicRtcSession. Disconnect returned %s"), *ToString(Result));
-			}
-		}
-
-		FUtf8String Utf8StreamerId(StreamerId);
-		EpicRtcManager->EpicRtcConference->RemoveSession(ToEpicRtcStringView(Utf8StreamerId));
-
-		EpicRtcManager->EpicRtcSession = nullptr;
-	}
-
-	void FStreamer::RemoveRoom()
-	{
-		if (!EpicRtcManager->EpicRtcRoom)
-		{
-			return;
-		}
-
-		FUtf8String Utf8StreamerId(StreamerId);
-		EpicRtcManager->EpicRtcRoom->Leave();
-		EpicRtcManager->EpicRtcSession->RemoveRoom(ToEpicRtcStringView(Utf8StreamerId));
-
-		EpicRtcManager->EpicRtcRoom = nullptr;
-	}
-
-	void FStreamer::OnStatsReady(const FString& PlayerId, const EpicRtcConnectionStats& ConnectionStats)
-	{
-		FPlayerContext* PlayerContext = Players->Find(PlayerId);
-		if (!PlayerContext)
-		{
-			return;
-		}
-
-		if (!PlayerContext->StatsCollector)
-		{
-			return;
-		}
-
-		PlayerContext->StatsCollector->Process(ConnectionStats);
-	}
-
-	IPixelStreaming2Streamer::FPreConnectionEvent& FStreamer::OnPreConnection()
-	{
-		return StreamingPreConnectionEvent;
-	}
-
-	IPixelStreaming2Streamer::FStreamingStartedEvent& FStreamer::OnStreamingStarted()
-	{
-		return StreamingStartedEvent;
-	}
-
-	IPixelStreaming2Streamer::FStreamingStoppedEvent& FStreamer::OnStreamingStopped()
-	{
-		return StreamingStoppedEvent;
-	}
-
-	void FStreamer::ForceKeyFrame()
-	{
-		VideoSourceGroup->ForceKeyFrame();
-	}
-
-	void FStreamer::FreezeStream(UTexture2D* Texture)
-	{
-		FreezeFrame->StartFreeze(Texture);
-	}
-
-	void FStreamer::UnfreezeStream()
-	{
-		// Force a keyframe so when stream unfreezes if player has never received a frame before they can still connect.
-		ForceKeyFrame();
-		FreezeFrame->StopFreeze();
-	}
-
-	void FStreamer::SendAllPlayersMessage(FString MessageType, const FString& Descriptor)
-	{
-		Players->Apply([&MessageType, &Descriptor](FString PlayerId, FPlayerContext& PlayerContext) {
-			if (PlayerContext.DataTrack && !IsSFU(PlayerId))
-			{
-				PlayerContext.DataTrack->SendMessage(MessageType, Descriptor);
-			}
-		});
-	}
-
-	void FStreamer::SendPlayerMessage(FString PlayerId, FString MessageType, const FString& Descriptor)
-	{
-		if (IsSFU(PlayerId))
-		{
-			return;
-		}
-		if (FPlayerContext* PlayerContext = Players->Find(PlayerId))
-		{
-			if (!PlayerContext->DataTrack)
-			{
-				return;
-			}
-			PlayerContext->DataTrack->SendMessage(MessageType, Descriptor);
-		}
-	}
-
-	void FStreamer::SendFileData(const TArray64<uint8>& ByteData, FString& MimeType, FString& FileExtension)
-	{
-		// TODO this should be dispatched as an async task, but because we lock when we visit the data
-		// channels it might be a bad idea. At some point it would be good to take a snapshot of the
-		// keys in the map when we start, then one by one get the channel and send the data
-
-		Players->Apply([&ByteData, &MimeType, &FileExtension](FString PlayerId, FPlayerContext& PlayerContext) {
-			if (!PlayerContext.DataTrack)
-			{
-				return;
-			}
-
-			// Send the mime type first
-			PlayerContext.DataTrack->SendMessage(EPixelStreaming2FromStreamerMessage::FileMimeType, MimeType);
-
-			// Send the extension next
-			PlayerContext.DataTrack->SendMessage(EPixelStreaming2FromStreamerMessage::FileExtension, FileExtension);
-
-			// Send the contents of the file. Note to callers: consider running this on its own thread, it can take a while if the file is big.
-			PlayerContext.DataTrack->SendArbitraryData(EPixelStreaming2FromStreamerMessage::FileContents, ByteData);
-		});
+	class FDummyStreamer : public IPixelStreaming2Streamer
+	{
+	public:
+		FDummyStreamer() = default;
+		virtual ~FDummyStreamer() = default;
+
+		virtual void									Initialize() override {}
+		virtual void									SetStreamFPS(int32 InFramesPerSecond) override {}
+		virtual int32									GetStreamFPS() override { return 0; }
+		virtual void									SetCoupleFramerate(bool bCouple) override {}
+		virtual void									SetVideoProducer(TSharedPtr<IPixelStreaming2VideoProducer> Input) override {}
+		virtual TWeakPtr<IPixelStreaming2VideoProducer> GetVideoProducer() override { return nullptr; }
+		virtual void									SetConnectionURL(const FString& InSignallingServerURL) override {}
+		virtual FString									GetConnectionURL() override { return FString(); }
+		virtual FString									GetId() override { return TEXT("DummyStreamer"); }
+		virtual bool									IsConnected() override { return false; }
+		virtual void									StartStreaming() override {}
+		virtual void									StopStreaming() override {}
+		virtual bool									IsStreaming() const { return false; }
+		virtual FPreConnectionEvent&					OnPreConnection() override { return StreamingPreConnectionEvent; }
+		virtual FStreamingStartedEvent&					OnStreamingStarted() override { return StreamingStartedEvent; }
+		virtual FStreamingStoppedEvent&					OnStreamingStopped() override { return StreamingStoppedEvent; }
+		virtual void									ForceKeyFrame() override {}
+		virtual void									FreezeStream(UTexture2D* Texture) override {}
+		virtual void									UnfreezeStream() override {}
+		virtual void									SendAllPlayersMessage(FString MessageType, const FString& Descriptor) override {}
+		virtual void									SendPlayerMessage(FString PlayerId, FString MessageType, const FString& Descriptor) override {}
+		virtual void									SendFileData(const TArray64<uint8>& ByteData, FString& MimeType, FString& FileExtension) override {}
+		virtual void									KickPlayer(FString PlayerId) override {}
+		virtual TArray<FString>							GetConnectedPlayers() override { return {}; }
+		virtual TWeakPtr<IPixelStreaming2InputHandler>	GetInputHandler() override { return nullptr; }
+		virtual IPixelStreaming2AudioSink*				GetPeerAudioSink(FString PlayerId) override { return nullptr; }
+		virtual IPixelStreaming2AudioSink*				GetUnlistenedAudioSink() override { return nullptr; }
+		virtual IPixelStreaming2VideoSink*				GetPeerVideoSink(FString PlayerId) override { return nullptr; }
+		virtual IPixelStreaming2VideoSink*				GetUnwatchedVideoSink() override { return nullptr; }
+		virtual void									SetConfigOption(const FName& OptionName, const FString& Value) override {}
+		virtual bool									GetConfigOption(const FName& OptionName, FString& OutValue) override { return false; }
+		virtual void									PlayerRequestsBitrate(FString PlayerId, int MinBitrate, int MaxBitrate) override {}
+		virtual void									RefreshStreamBitrate() override {}
+
+	private:
+		FPreConnectionEvent	   StreamingPreConnectionEvent;
+		FStreamingStartedEvent StreamingStartedEvent;
+		FStreamingStoppedEvent StreamingStoppedEvent;
+	};
+
+	FString FDummyStreamerFactory::GetStreamType()
+	{
+		return FString();
+	}
+
+	TSharedPtr<IPixelStreaming2Streamer> FDummyStreamerFactory::CreateNewStreamer(const FString& StreamerId)
+	{
+		return TSharedPtr<FDummyStreamer>(new FDummyStreamer());
 	}
+} // namespace UE::PixelStreaming2
 
-	void FStreamer::KickPlayer(FString PlayerId)
-	{
-		if (FPlayerContext* PlayerContext = Players->Find(PlayerId))
-		{
-			PlayerContext->ParticipantInterface->Kick();
-		}
-	}
+IPixelStreaming2StreamerFactory::IPixelStreaming2StreamerFactory()
+{
+	IPixelStreaming2StreamerFactory::RegisterStreamerFactory(this);
+}
 
-	TArray<FString> FStreamer::GetConnectedPlayers()
-	{
-		TArray<FString> ConnectedPlayerIds;
-		Players->Apply([&ConnectedPlayerIds, this](FString PlayerId, FPlayerContext& PlayerContext) {
-			ConnectedPlayerIds.Add(PlayerId);
-		});
-		return ConnectedPlayerIds;
-	}
+IPixelStreaming2StreamerFactory::~IPixelStreaming2StreamerFactory()
+{
+	IPixelStreaming2StreamerFactory::UnregisterStreamerFactory(this);
+}
 
-	IPixelStreaming2AudioSink* FStreamer::GetPeerAudioSink(FString PlayerId)
-	{
-		if (FPlayerContext* PlayerContext = Players->Find(PlayerId))
-		{
-			if (PlayerContext->AudioSink)
-			{
-				return PlayerContext->AudioSink.Get();
-			}
-		}
-		return nullptr;
-	}
+void IPixelStreaming2StreamerFactory::RegisterStreamerFactory(IPixelStreaming2StreamerFactory* InFactory)
+{
+	IModularFeatures::Get().RegisterModularFeature(ModularFeatureName, InFactory);
+}
 
-	IPixelStreaming2AudioSink* FStreamer::GetUnlistenedAudioSink()
-	{
-		IPixelStreaming2AudioSink* Result = nullptr;
-		Players->ApplyUntil([&Result](FString PlayerId, FPlayerContext& PlayerContext) {
-			if (PlayerContext.AudioSink)
-			{
-				if (!PlayerContext.AudioSink->HasAudioConsumers())
-				{
-					Result = PlayerContext.AudioSink.Get();
-					return true;
-				}
-			}
-			return false;
-		});
-		return Result;
-	}
+void IPixelStreaming2StreamerFactory::UnregisterStreamerFactory(IPixelStreaming2StreamerFactory* InFactory)
+{
+	IModularFeatures::Get().UnregisterModularFeature(ModularFeatureName, InFactory);
+}
 
-	IPixelStreaming2VideoSink* FStreamer::GetPeerVideoSink(FString PlayerId)
+IPixelStreaming2StreamerFactory* IPixelStreaming2StreamerFactory::Get(const FString& InType)
+{
+	if (InType == FString())
 	{
-		if (FPlayerContext* PlayerContext = Players->Find(PlayerId))
-		{
-			if (PlayerContext->VideoSink)
-			{
-				return PlayerContext->VideoSink.Get();
-			}
-		}
 		return nullptr;
 	}
 
-	IPixelStreaming2VideoSink* FStreamer::GetUnwatchedVideoSink()
-	{
-		IPixelStreaming2VideoSink* Result = nullptr;
-		Players->ApplyUntil([&Result](FString PlayerId, FPlayerContext& PlayerContext) {
-			if (PlayerContext.VideoSink)
-			{
-				if (!PlayerContext.VideoSink->HasVideoConsumers())
-				{
-					Result = PlayerContext.VideoSink.Get();
-					return true;
-				}
-			}
-			return false;
-		});
-		return Result;
-	}
-
-	void FStreamer::SetConfigOption(const FName& OptionName, const FString& Value)
-	{
-		if (Value.IsEmpty())
-		{
-			ConfigOptions.Remove(OptionName);
-		}
-		else
-		{
-			ConfigOptions.Add(OptionName, Value);
-		}
-	}
-
-	bool FStreamer::GetConfigOption(const FName& OptionName, FString& OutValue)
-	{
-		FString* OptionValue = ConfigOptions.Find(OptionName);
-		if (OptionValue)
-		{
-			OutValue = *OptionValue;
-			return true;
-		}
-		else
-		{
-			return false;
-		}
-	}
-
-	void FStreamer::PlayerRequestsBitrate(FString PlayerId, int MinBitrate, int MaxBitrate)
-	{
-		UPixelStreaming2PluginSettings::CVarWebRTCMinBitrate.AsVariable()->Set(MinBitrate);
-		UPixelStreaming2PluginSettings::CVarWebRTCMaxBitrate.AsVariable()->Set(MaxBitrate);
-	}
-
-	void FStreamer::RefreshStreamBitrate()
-	{
-		Players->Apply([this](FString PlayerId, FPlayerContext& PlayerContext) {
-			if (!PlayerContext.ParticipantInterface)
-			{
-				return;
-			}
-
-			TRefCountPtr<EpicRtcConnectionInterface> ConnectionInterface = PlayerContext.ParticipantInterface->GetConnection();
-			if (!ConnectionInterface)
-			{
-				return;
-			}
-
-			EpicRtcBitrate Bitrates = {
-				._minBitrateBps = UPixelStreaming2PluginSettings::CVarWebRTCMinBitrate.GetValueOnAnyThread(),
-				._hasMinBitrateBps = true,
-				._maxBitrateBps = UPixelStreaming2PluginSettings::CVarWebRTCMaxBitrate.GetValueOnAnyThread(),
-				._hasMaxBitrateBps = true,
-				._startBitrateBps = UPixelStreaming2PluginSettings::CVarWebRTCStartBitrate.GetValueOnAnyThread(),
-				._hasStartBitrateBps = true
-			};
-
-			ConnectionInterface->SetConnectionRates(Bitrates);
-		});
-	}
-
-	void FStreamer::ForEachPlayer(const TFunction<void(FString, FPlayerContext)>& Func)
-	{
-		Players->Apply(Func);
-	}
-
-	void FStreamer::ConsumeStats(FString PlayerId, FName StatName, float StatValue)
-	{
-		if (IsSFU(PlayerId))
-		{
-			return;
-		}
-
-		if (StatName != PixelStreaming2StatNames::MeanQPPerSecond)
-		{
-			return;
-		}
-
-		FPlayerContext* PlayerContext = Players->Find(PlayerId);
-
-		if (!PlayerContext)
-		{
-			return;
-		}
-
-		if (!PlayerContext->DataTrack)
-		{
-			return;
-		}
-
-		PlayerContext->DataTrack->SendMessage(EPixelStreaming2FromStreamerMessage::VideoEncoderAvgQP, FString::FromInt((int)StatValue));
-	}
-
-	void FStreamer::DeletePlayerSession(FString PlayerId)
-	{
-		// We dont want to allow this to be deleted within Players.Remove because
-		// we lock the players map and the delete could dispatch a webrtc object
-		// delete on the signalling thread which might be waiting for the players
-		// lock.
-		FPlayerContext PendingDeletePlayer;
-		if (FPlayerContext* PlayerContext = Players->Find(PlayerId))
-		{
-			// when a sfu is connected we only get disconnect messages.
-			// we dont get connect messages but we might get datachannel requests which can result
-			// in players with no PeerConnection but a datachannel
-			if (PlayerContext->VideoSource)
-			{
-				VideoSourceGroup->RemoveVideoSource(PlayerContext->VideoSource.Get());
-			}
-
-			// Close any data track related things (we do this here because RoomLeft happens before DataTrack stopped fires)
-			// So if we only did this in DataTrack stooped the playerId would already be removed
-			OnDataChannelClosed(PlayerId);
-
-			PendingDeletePlayer = *PlayerContext;
-		}
-
-		Players->Remove(PlayerId);
-
-		if (UPixelStreaming2Delegates* Delegates = UPixelStreaming2Delegates::Get())
-		{
-			Delegates->OnClosedConnection.Broadcast(StreamerId, PlayerId);
-			Delegates->OnClosedConnectionNative.Broadcast(StreamerId, PlayerId);
-			if (Players->IsEmpty())
-			{
-				Delegates->OnAllConnectionsClosed.Broadcast(StreamerId);
-				Delegates->OnAllConnectionsClosedNative.Broadcast(StreamerId);
-			}
-		}
-
-		if (FStats* PSStats = FStats::Get())
-		{
-			PSStats->RemovePeerStats(PlayerId);
-		}
-	}
-
-	void FStreamer::DeleteAllPlayerSessions()
-	{
-		if (FStats* PSStats = FStats::Get())
-		{
-			PSStats->RemoveAllPeerStats();
-		}
-
-		VideoSourceGroup->RemoveAllVideoSources();
-		Players->Empty();
-		InputControllingId = INVALID_PLAYER_ID;
-		if (UPixelStreaming2Delegates* Delegates = UPixelStreaming2Delegates::Get())
-		{
-			Delegates->OnAllConnectionsClosed.Broadcast(StreamerId);
-			Delegates->OnAllConnectionsClosedNative.Broadcast(StreamerId);
-		}
-	}
-
-	void FStreamer::OnDataChannelOpen(FString PlayerId)
-	{
-		// Only time we automatically make a new peer the input controlling host is if they are the first peer (and not the SFU).
-		bool HostControlsInput = GetEnumFromCVar<EInputControllerMode>(UPixelStreaming2PluginSettings::CVarInputController) == EInputControllerMode::Host;
-		if (HostControlsInput && !IsSFU(PlayerId) && InputControllingId == INVALID_PLAYER_ID)
-		{
-			InputControllingId = PlayerId;
-		}
-
-		if (UPixelStreaming2Delegates* Delegates = UPixelStreaming2Delegates::Get())
-		{
-			Delegates->OnDataTrackOpen.Broadcast(StreamerId, PlayerId);
-			Delegates->OnDataTrackOpenNative.Broadcast(StreamerId, PlayerId);
-		}
-
-		// When data channel is open
-		SendProtocol(PlayerId);
-		// Try to send cached freeze frame (if we have one)
-		FreezeFrame->SendCachedFreezeFrameTo(PlayerId);
-		SendInitialSettings(PlayerId);
-		SendPeerControllerMessages(PlayerId);
-	}
+	IModularFeatures::Get().LockModularFeatureList();
+	TArray<IPixelStreaming2StreamerFactory*> Factories = IModularFeatures::Get().GetModularFeatureImplementations<IPixelStreaming2StreamerFactory>(ModularFeatureName);
+	IModularFeatures::Get().UnlockModularFeatureList();
 
-	void FStreamer::OnDataChannelClosed(FString PlayerId)
+	for (IPixelStreaming2StreamerFactory* Factory : Factories)
 	{
-		if (FPlayerContext* PlayerContext = Players->Find(PlayerId))
+		if (Factory && InType == Factory->GetStreamType())
 		{
-			PlayerContext->DataTrack = nullptr;
-
-			if (InputControllingId == PlayerId)
-			{
-				InputControllingId = INVALID_PLAYER_ID;
-				// just get the first channel we have and give it input control.
-				Players->ApplyUntil([this](FString PlayerId, FPlayerContext& PlayerContext) {
-					if (!PlayerContext.DataTrack)
-					{
-						return false;
-					}
-					if (IsSFU(PlayerId))
-					{
-						return false;
-					}
-					InputControllingId = PlayerId;
-					PlayerContext.DataTrack->SendMessage(EPixelStreaming2FromStreamerMessage::InputControlOwnership, 1 /* ControlsInput */);
-					return true;
-				});
-			}
-
-			if (UPixelStreaming2Delegates* Delegates = UPixelStreaming2Delegates::Get())
-			{
-				Delegates->OnDataTrackClosed.Broadcast(StreamerId, PlayerId);
-				Delegates->OnDataTrackClosedNative.Broadcast(StreamerId, PlayerId);
-			}
+			return Factory;
 		}
 	}
 
-	void FStreamer::SendInitialSettings(FString PlayerId) const
-	{
-		const FString PixelStreaming2Payload = FString::Printf(TEXT("{ \"AllowPixelStreamingCommands\": %s, \"DisableLatencyTest\": %s }"),
-			UPixelStreaming2PluginSettings::CVarInputAllowConsoleCommands.GetValueOnAnyThread() ? TEXT("true") : TEXT("false"),
-			UPixelStreaming2PluginSettings::CVarDisableLatencyTester.GetValueOnAnyThread() ? TEXT("true") : TEXT("false"));
-
-		const FString WebRTCPayload = FString::Printf(TEXT("{ \"FPS\": %d, \"MinBitrate\": %d, \"MaxBitrate\": %d }"),
-			UPixelStreaming2PluginSettings::CVarWebRTCFps.GetValueOnAnyThread(),
-			UPixelStreaming2PluginSettings::CVarWebRTCMinBitrate.GetValueOnAnyThread(),
-			UPixelStreaming2PluginSettings::CVarWebRTCMaxBitrate.GetValueOnAnyThread());
+	UE_LOGFMT(LogPixelStreaming2, Warning, "No streamer factory implementation for {0} found. Streamers set to this type will not do anything.", *InType);
+	static UE::PixelStreaming2::FDummyStreamerFactory DummyFactory = UE::PixelStreaming2::FDummyStreamerFactory();
+	return &DummyFactory;
+}
 
-		const FString EncoderPayload = FString::Printf(TEXT("{ \"TargetBitrate\": %d, \"MinQuality\": %d, \"MaxQuality\": %d }"),
-			UPixelStreaming2PluginSettings::CVarEncoderTargetBitrate.GetValueOnAnyThread(),
-			UPixelStreaming2PluginSettings::CVarEncoderMinQuality.GetValueOnAnyThread(),
-			UPixelStreaming2PluginSettings::CVarEncoderMaxQuality.GetValueOnAnyThread());
-
-		FString ConfigPayload = TEXT("{ ");
-		bool	bComma = false; // Simplest way to avoid complaints from pedantic JSON parsers
-		for (const TPair<FName, FString>& Option : ConfigOptions)
-		{
-			if (bComma)
-			{
-				ConfigPayload.Append(TEXT(", "));
-			}
-			ConfigPayload.Append(FString::Printf(TEXT("\"%s\": \"%s\""), *Option.Key.ToString(), *Option.Value));
-			bComma = true;
-		}
-		ConfigPayload.Append(TEXT("}"));
-
-		const FString FullPayload = FString::Printf(TEXT("{ \"PixelStreaming\": %s, \"Encoder\": %s, \"WebRTC\": %s, \"ConfigOptions\": %s }"), *PixelStreaming2Payload, *EncoderPayload, *WebRTCPayload, *ConfigPayload);
-
-		if (const FPlayerContext* PlayerContext = Players->Find(PlayerId))
-		{
-			if (!PlayerContext->DataTrack)
-			{
-				return;
-			}
-			PlayerContext->DataTrack->SendMessage(EPixelStreaming2FromStreamerMessage::InitialSettings, FullPayload);
-		}
-	}
-
-	void FStreamer::SendProtocol(FString PlayerId) const
-	{
-		const TArray<TSharedPtr<IPixelStreaming2DataProtocol>> Protocols = { InputHandler->GetToStreamerProtocol(), InputHandler->GetFromStreamerProtocol() };
-		for (TSharedPtr<IPixelStreaming2DataProtocol> Protocol : Protocols)
-		{
-			TSharedPtr<FJsonObject>	  ProtocolJson = Protocol->ToJson();
-			FString					  Body;
-			TSharedRef<TJsonWriter<>> JsonWriter = TJsonWriterFactory<>::Create(&Body);
-			if (!ensure(FJsonSerializer::Serialize(ProtocolJson.ToSharedRef(), JsonWriter)))
-			{
-				UE_LOG(LogPixelStreaming2, Warning, TEXT("Cannot serialize protocol json object"));
-				return;
-			}
+TArray<FString> IPixelStreaming2StreamerFactory::GetAvailableFactoryTypes()
+{
+	TArray<FString> StreamProtocols;
 
-			if (const FPlayerContext* PlayerContext = Players->Find(PlayerId))
-			{
-				if (!PlayerContext->DataTrack)
-				{
-					return;
-				}
-				PlayerContext->DataTrack->SendMessage(EPixelStreaming2FromStreamerMessage::Protocol, Body);
-			}
-		}
-	}
+	IModularFeatures::Get().LockModularFeatureList();
+	TArray<IPixelStreaming2StreamerFactory*> Factories = IModularFeatures::Get().GetModularFeatureImplementations<IPixelStreaming2StreamerFactory>(ModularFeatureName);
+	IModularFeatures::Get().UnlockModularFeatureList();
 
-	void FStreamer::SendPeerControllerMessages(FString PlayerId) const
+	for (IPixelStreaming2StreamerFactory* Factory : Factories)
 	{
-		if (const FPlayerContext* PlayerContext = Players->Find(PlayerId))
-		{
-			if (!PlayerContext->DataTrack)
-			{
-				return;
-			}
-			const uint8 ControlsInput = (GetEnumFromCVar<EInputControllerMode>(UPixelStreaming2PluginSettings::CVarInputController) == EInputControllerMode::Host) ? (PlayerId == InputControllingId) : 1;
-			// Even though the QualityController feature is removed we send it for backwards compatibility with older frontends (can probably remove 2 versions after 5.5)
-			PlayerContext->DataTrack->SendMessage(EPixelStreaming2FromStreamerMessage::InputControlOwnership, ControlsInput);
-			PlayerContext->DataTrack->SendMessage(EPixelStreaming2FromStreamerMessage::QualityControlOwnership, 1 /* True */);
-		}
+		StreamProtocols.AddUnique(Factory->GetStreamType());
 	}
 
-	void FStreamer::SendLatencyReport(FString PlayerId) const
-	{
-		if (UPixelStreaming2PluginSettings::CVarDisableLatencyTester.GetValueOnAnyThread())
-		{
-			return;
-		}
-
-		double ReceiptTimeMs = FPlatformTime::ToMilliseconds64(FPlatformTime::Cycles64());
-
-		AsyncTask(ENamedThreads::GameThread, [this, PlayerId, ReceiptTimeMs]() {
-			FString ReportToTransmitJSON;
-
-			if (!UPixelStreaming2PluginSettings::CVarWebRTCDisableStats.GetValueOnAnyThread())
-			{
-				double EncodeMs = -1.0;
-				double CaptureToSendMs = 0.0;
-
-				FStats* Stats = FStats::Get();
-				if (Stats)
-				{
-					Stats->QueryPeerStat(PlayerId, RTCStatCategories::LocalVideoTrack, PixelStreaming2StatNames::MeanEncodeTime, EncodeMs);
-					Stats->QueryPeerStat(PlayerId, RTCStatCategories::LocalVideoTrack, PixelStreaming2StatNames::AvgSendDelay, CaptureToSendMs);
-				}
-
-				double TransmissionTimeMs = FPlatformTime::ToMilliseconds64(FPlatformTime::Cycles64());
-				ReportToTransmitJSON = FString::Printf(
-					TEXT("{ \"ReceiptTimeMs\": %.2f, \"EncodeMs\": %.2f, \"CaptureToSendMs\": %.2f, \"TransmissionTimeMs\": %.2f }"),
-					ReceiptTimeMs,
-					EncodeMs,
-					CaptureToSendMs,
-					TransmissionTimeMs);
-			}
-			else
-			{
-				double TransmissionTimeMs = FPlatformTime::ToMilliseconds64(FPlatformTime::Cycles64());
-				ReportToTransmitJSON = FString::Printf(
-					TEXT("{ \"ReceiptTimeMs\": %.2f, \"EncodeMs\": \"Pixel Streaming stats are disabled\", \"CaptureToSendMs\": \"Pixel Streaming stats are disabled\", \"TransmissionTimeMs\": %.2f }"),
-					ReceiptTimeMs,
-					TransmissionTimeMs);
-			}
-
-			if (const FPlayerContext* PlayerContext = Players->Find(PlayerId))
-			{
-				if (PlayerContext->DataTrack)
-				{
-					PlayerContext->DataTrack->SendMessage(EPixelStreaming2FromStreamerMessage::LatencyTest, ReportToTransmitJSON);
-				}
-			}
-		});
-	}
-
-	void FStreamer::HandleRelayStatusMessage(const uint8_t* Data, uint32_t Size, EpicRtcDataTrackInterface* DataTrack)
-	{
-		//skip type
-		Data++;
-		Size--;
-		FString PlayerId = ReadString(Data, Size);
-		checkf(Size > 0, TEXT("Malformed relay status message!"))
-		bool bIsOn = static_cast<bool>(Data[0]);
-
-		FString DataTrackId = ToString(DataTrack->GetId());
-		if (bIsOn)
-		{
-			UE_LOG(LogPixelStreaming2, Log, TEXT("FStreamer::HandleRelayStatusMessage() Adding new PlayerId [%s] with DataTrackId [%s]"), *PlayerId, *DataTrackId);
-			
-			FString SFUId;
-			if (FindPlayerByDataTrack(DataTrack, SFUId))
-			{
-				FPlayerContext* SFUContext = Players->Find(SFUId);
-				FPlayerContext& PlayerContext = Players->FindOrAdd(PlayerId);
-				PlayerContext.DataTrack = FEpicRtcMutliplexDataTrack::Create(SFUContext->DataTrack, InputHandler->GetFromStreamerProtocol(), PlayerId);
-				OnDataChannelOpen(PlayerId);
-			}
-			else
-			{
-				UE_LOG(LogPixelStreaming2, Log, TEXT("FStreamer::HandleRelayStatusMessage() Failed to find SFU PlayerContext"));
-			}
-		}
-		else
-		{
-			UE_LOG(LogPixelStreaming2, Log, TEXT("FStreamer::HandleRelayStatusMessage() Removing PlayerId [%s] with DataTrackId [%s]"), *PlayerId, *DataTrackId);
-
-			OnDataChannelClosed(PlayerId);
-			Players->Remove(PlayerId);
-		}
-	}
-
-	void FStreamer::TriggerMouseLeave(FString InStreamerId)
-	{
-		if (!IsEngineExitRequested() && StreamerId == InStreamerId)
-		{
-			TSharedPtr<IPixelStreaming2InputHandler> SharedInputHandler = InputHandler;
-
-			// Force a MouseLeave event. This prevents the PixelStreaming2ApplicationWrapper from
-			// still wrapping the base FSlateApplication after we stop streaming
-			const auto MouseLeaveFunction = [SharedInputHandler]() {
-				if (SharedInputHandler.IsValid())
-				{
-					TArray<uint8>							EmptyArray;
-					TFunction<void(FString, FMemoryReader)> MouseLeaveHandler = SharedInputHandler->FindMessageHandler("MouseLeave");
-					MouseLeaveHandler("", FMemoryReader(EmptyArray));
-				}
-			};
-
-			if (IsInGameThread())
-			{
-				MouseLeaveFunction();
-			}
-			else
-			{
-				AsyncTask(ENamedThreads::GameThread, [MouseLeaveFunction]() {
-					MouseLeaveFunction();
-				});
-			}
-		}
-	}
-
-	void FStreamer::OnSessionStateUpdate(const EpicRtcSessionState State)
-	{
-		switch (State)
-		{
-			case EpicRtcSessionState::Connected:
-			{
-				bSignallingConnected = true;
-				if (UPixelStreaming2Delegates* Delegates = UPixelStreaming2Delegates::Get())
-				{
-					Delegates->OnConnectedToSignallingServer.Broadcast(StreamerId);
-					Delegates->OnConnectedToSignallingServerNative.Broadcast(StreamerId);
-				}
-
-				UE_LOG(LogPixelStreaming2, VeryVerbose, TEXT("FStreamer::OnSessionStateUpdate State=Connected"));
-				EpicRtcBitrate Bitrate = {
-					._minBitrateBps = UPixelStreaming2PluginSettings::CVarWebRTCMinBitrate.GetValueOnAnyThread(),
-					._hasMinBitrateBps = true,
-					._maxBitrateBps = UPixelStreaming2PluginSettings::CVarWebRTCMaxBitrate.GetValueOnAnyThread(),
-					._hasMaxBitrateBps = true,
-					._startBitrateBps = UPixelStreaming2PluginSettings::CVarWebRTCStartBitrate.GetValueOnAnyThread(),
-					._hasStartBitrateBps = true
-				};
-
-				EpicRtcPortAllocator PortAllocator = {
-					._minPort = UPixelStreaming2PluginSettings::CVarWebRTCMinPort.GetValueOnAnyThread(),
-					._hasMinPort = true,
-					._maxPort = UPixelStreaming2PluginSettings::CVarWebRTCMaxPort.GetValueOnAnyThread(),
-					._hasMaxPort = true,
-					._portAllocation = static_cast<EpicRtcPortAllocatorOptions>(UPixelStreaming2PluginSettings::GetPortAllocationFlags())
-				};
-
-				EpicRtcConnectionConfig ConnectionConfig = {
-					._iceServers = { ._ptr = nullptr, ._size = 0 }, // This can stay empty because EpicRtc handles the ice servers internally
-					._portAllocator = PortAllocator,
-					._bitrate = Bitrate,
-					._iceConnectionPolicy = EpicRtcIcePolicy::All,
-					._disableTcpCandidates = false
-				};
-
-				FUtf8String		  Utf8StreamerId(StreamerId);
-				EpicRtcRoomConfig RoomConfig = {
-					._id = ToEpicRtcStringView(Utf8StreamerId),
-					._connectionConfig = ConnectionConfig,
-					._ticket = EpicRtcStringView{ ._ptr = nullptr, ._length = 0 },
-					._observer = EpicRtcManager->RoomObserver,
-					._audioTrackObserverFactory = EpicRtcManager->AudioTrackObserverFactory,
-					._dataTrackObserverFactory = EpicRtcManager->DataTrackObserverFactory,
-					._videoTrackObserverFactory = EpicRtcManager->VideoTrackObserverFactory
-				};
-
-				EpicRtcErrorCode Result = EpicRtcManager->EpicRtcSession->CreateRoom(RoomConfig, EpicRtcManager->EpicRtcRoom.GetInitReference());
-				if (Result != EpicRtcErrorCode::Ok)
-				{
-					UE_LOG(LogPixelStreaming2, Error, TEXT("Failed to create EpicRtc room. CreateRoom returned %s"), *ToString(Result));
-					break;
-				}
-
-				EpicRtcManager->EpicRtcRoom->Join();
-
-				// Would be better renamed to OnSessionConnected
-				OnStreamingStarted().Broadcast(this);
-			}
-			break;
-			case EpicRtcSessionState::New:
-				// Do something on `new` Session here
-				UE_LOG(LogPixelStreaming2, VeryVerbose, TEXT("FStreamer::OnSessionStateUpdate State=New"));
-				break;
-			case EpicRtcSessionState::Pending:
-				// Do something on `pending` session here
-				UE_LOG(LogPixelStreaming2, VeryVerbose, TEXT("FStreamer::OnSessionStateUpdate State=Pending"));
-				break;
-			case EpicRtcSessionState::Disconnected:
-				// Do something on `disconnected` session here
-				bSignallingConnected = false;
-				if (UPixelStreaming2Delegates* Delegates = UPixelStreaming2Delegates::Get())
-				{
-					Delegates->OnDisconnectedFromSignallingServer.Broadcast(StreamerId);
-					Delegates->OnDisconnectedFromSignallingServerNative.Broadcast(StreamerId);
-				}
-				UE_LOG(LogPixelStreaming2, VeryVerbose, TEXT("FStreamer::OnSessionStateUpdate State=Disconnected"));
-				RemoveSession(false);
-				StopStreaming();
-				ReconnectTimer->Start(AsShared());
-				break;
-			case EpicRtcSessionState::Failed:
-				// Do something on `failed` session here
-				bSignallingConnected = false;
-				UE_LOG(LogPixelStreaming2, VeryVerbose, TEXT("FStreamer::OnSessionStateUpdate State=Failed"));
-				break;
-			case EpicRtcSessionState::Exiting:
-				// Do something on `exiting` session here
-				UE_LOG(LogPixelStreaming2, VeryVerbose, TEXT("FStreamer::OnSessionStateUpdate State=Exiting"));
-				break;
-			default:
-				UE_LOG(LogPixelStreaming2, Error, TEXT("FStreamer::OnSessionStateUpdate An unhandled session state was encountered. This switch might be missing a case."));
-				break;
-		}
-	}
-
-	void FStreamer::OnSessionErrorUpdate(const EpicRtcErrorCode ErrorUpdate)
-	{
-		UE_LOGFMT(LogPixelStreaming2, VeryVerbose, "FStreamer::OnSessionErrorUpdate does nothing");
-	}
-
-	void FStreamer::OnSessionRoomsAvailableUpdate(EpicRtcStringArrayInterface* RoomsList)
-	{
-		UE_LOGFMT(LogPixelStreaming2, VeryVerbose, "FStreamer::OnSessionRoomsAvailableUpdate does nothing");
-	}
-
-	void FStreamer::OnRoomStateUpdate(const EpicRtcRoomState State)
-	{
-		UE_LOGFMT(LogPixelStreaming2, VeryVerbose, "FStreamer::OnRoomStateUpdate does nothing");
-	}
-
-	FUtf8String GetAudioStreamID()
-	{
-		const bool bSyncVideoAndAudio = !UPixelStreaming2PluginSettings::CVarWebRTCDisableAudioSync.GetValueOnAnyThread();
-		return bSyncVideoAndAudio ? "pixelstreaming_av_stream_id" : "pixelstreaming_audio_stream_id";
-	}
-
-	FUtf8String GetVideoStreamID()
-	{
-		const bool bSyncVideoAndAudio = !UPixelStreaming2PluginSettings::CVarWebRTCDisableAudioSync.GetValueOnAnyThread();
-		return bSyncVideoAndAudio ? "pixelstreaming_av_stream_id" : "pixelstreaming_video_stream_id";
-	}
-
-	void FStreamer::OnRoomJoinedUpdate(EpicRtcParticipantInterface* Participant)
-	{
-		FString ParticipantId = ToString(Participant->GetId());
-		UE_LOG(LogPixelStreaming2, Log, TEXT("Player (%s) joined"), *ParticipantId);
-
-		if (ParticipantId == StreamerId)
-		{
-			return;
-		}
-
-		if (UPixelStreaming2Delegates* Delegates = UPixelStreaming2Delegates::Get())
-		{
-			Delegates->OnNewConnection.Broadcast(StreamerId, ParticipantId);
-			Delegates->OnNewConnectionNative.Broadcast(StreamerId, ParticipantId);
-		}
-
-		FPlayerContext& PlayerContext = Players->FindOrAdd(ParticipantId);
-		PlayerContext.ParticipantInterface = Participant;
-		PlayerContext.StatsCollector = FRTCStatsCollector::Create(ParticipantId);
-
-		TRefCountPtr<EpicRtcConnectionInterface> ParticipantConnection = Participant->GetConnection();
-		ParticipantConnection->SetManualNegotiation(true);
-
-		const EVideoCodec SelectedCodec = GetEnumFromCVar<EVideoCodec>(UPixelStreaming2PluginSettings::CVarEncoderCodec);
-		const bool		  bNegotiateCodecs = UPixelStreaming2PluginSettings::CVarWebRTCNegotiateCodecs.GetValueOnAnyThread();
-		const bool		  bTransmitUEVideo = !UPixelStreaming2PluginSettings::CVarWebRTCDisableTransmitVideo.GetValueOnAnyThread();
-		bool			  bReceiveBrowserVideo = !UPixelStreaming2PluginSettings::CVarWebRTCDisableReceiveVideo.GetValueOnAnyThread();
-
-		// Check if the user has selected only H.264 on a AMD gpu and disable receiving video.
-		// WebRTC does not support using SendRecv if the encoding and decoding do not support the same codec.
-		// AMD GPUs currently have decoding disabled so WebRTC fails to create SDP codecs with SendRecv.
-		if (IsRHIDeviceAMD() && !bNegotiateCodecs && SelectedCodec == EVideoCodec::H264)
-		{
-			if (bReceiveBrowserVideo)
-			{
-				bReceiveBrowserVideo = false;
-				UE_LOGFMT(LogPixelStreaming2, Warning, "AMD GPUs do not support receiving H.264 video.");
-			}
-		}
-
-		if (bTransmitUEVideo || bReceiveBrowserVideo)
-		{
-			TArray<EpicRtcVideoEncodingConfig> VideoEncodingConfigs;
-			// We need ensure the Rids have the same lifetime as the VideoEncodingConfigs
-			// to ensure the contents don't get deleted before we can call AddVideoSource
-			TArray<FUtf8String> Rids;
-
-			int MaxFramerate = UPixelStreaming2PluginSettings::CVarWebRTCFps.GetValueOnAnyThread();
-
-			TArray<FPixelStreaming2SimulcastLayer> SimulcastParams = UE::PixelStreaming2::GetSimulcastParameters();
-			for (int i = 0; i < SimulcastParams.Num(); ++i)
-			{
-				const FPixelStreaming2SimulcastLayer& SpatialLayer = SimulcastParams[i];
-
-				if (SimulcastParams.Num() > 1)
-				{
-					Rids.Add(FUtf8String("simulcast") + FUtf8String::FromInt(SimulcastParams.Num() - i));
-				}
-
-				EpicRtcVideoEncodingConfig VideoEncodingConfig = {
-					// clang-format off
-					// TODO (Migration): RTCP-7027 Maybe bug in EpicRtc? Setting an rid if there's only one config results in no video
-					._rid = SimulcastParams.Num() > 1
-						? EpicRtcStringView{._ptr = (const char*)(*(Rids[i])), ._length = static_cast<uint64>(Rids[i].Len()) }
-						: EpicRtcStringView{._ptr = nullptr, ._length = 0 },
-					// clang-format on
-					._scaleResolutionDownBy = SpatialLayer.Scaling,
-					._scalabilityMode = static_cast<EpicRtcVideoScalabilityMode>(GetEnumFromCVar<EScalabilityMode>(UPixelStreaming2PluginSettings::CVarEncoderScalabilityMode)), // HACK if the Enums become un-aligned
-					._minBitrate = (uint32_t)SpatialLayer.MinBitrate,
-					._maxBitrate = (uint32_t)SpatialLayer.MaxBitrate,
-					._maxFrameRate = (uint8_t)MaxFramerate
-				};
-
-				VideoEncodingConfigs.Add(VideoEncodingConfig);
-			}
-
-			EpicRtcVideoEncodingConfigSpan VideoEncodingConfigSpan = {
-				._ptr = VideoEncodingConfigs.GetData(),
-				._size = (uint64_t)VideoEncodingConfigs.Num()
-			};
-
-			EpicRtcMediaSourceDirection VideoDirection;
-			if (bTransmitUEVideo && bReceiveBrowserVideo)
-			{
-				VideoDirection = EpicRtcMediaSourceDirection::SendRecv;
-			}
-			else if (bTransmitUEVideo)
-			{
-				VideoDirection = EpicRtcMediaSourceDirection::SendOnly;
-			}
-			else if (bReceiveBrowserVideo)
-			{
-				VideoDirection = EpicRtcMediaSourceDirection::RecvOnly;
-			}
-			else
-			{
-				VideoDirection = EpicRtcMediaSourceDirection::RecvOnly;
-			}
-
-			FUtf8String		   VideoStreamID = GetVideoStreamID();
-			EpicRtcVideoSource VideoSource = {
-				._streamId = ToEpicRtcStringView(VideoStreamID),
-				._encodings = VideoEncodingConfigSpan,
-				._direction = VideoDirection
-			};
-
-			ParticipantConnection->AddVideoSource(VideoSource);
-		}
-
-		const bool bTransmitUEAudio = !UPixelStreaming2PluginSettings::CVarWebRTCDisableTransmitAudio.GetValueOnAnyThread();
-		const bool bReceiveBrowserAudio = !UPixelStreaming2PluginSettings::CVarWebRTCDisableReceiveAudio.GetValueOnAnyThread();
-		if (bTransmitUEAudio || bReceiveBrowserAudio)
-		{
-			EpicRtcMediaSourceDirection AudioDirection;
-			if (bTransmitUEAudio && bReceiveBrowserAudio)
-			{
-				AudioDirection = EpicRtcMediaSourceDirection::SendRecv;
-			}
-			else if (bTransmitUEAudio)
-			{
-				AudioDirection = EpicRtcMediaSourceDirection::SendOnly;
-			}
-			else if (bReceiveBrowserAudio)
-			{
-				AudioDirection = EpicRtcMediaSourceDirection::RecvOnly;
-			}
-			else
-			{
-				AudioDirection = EpicRtcMediaSourceDirection::RecvOnly;
-			}
-
-			FUtf8String		   AudioStreamID = GetAudioStreamID();
-			EpicRtcAudioSource AudioSource = {
-				._streamId = ToEpicRtcStringView(AudioStreamID),
-				._bitrate = 510000,
-				._channels = 2,
-				._direction = AudioDirection
-			};
-
-			ParticipantConnection->AddAudioSource(AudioSource);
-		}
-
-		if (IsSFU(ParticipantId))
-		{
-			FString RecvLabel(TEXT("recv-datachannel"));
-			FUtf8String Utf8RecvLabel = *RecvLabel;
-			EpicRtcDataSource RecvDataSource = {
-				._label = ToEpicRtcStringView(Utf8RecvLabel),
-				._maxRetransmitTime = 0,
-				._maxRetransmits = 0,
-				._isOrdered = true,
-				._protocol = EpicRtcDataSourceProtocol::Sctp,
-				._negotiated = true,
-				._transportChannelId = 1
-			};
-			ParticipantConnection->AddDataSource(RecvDataSource);
-
-			FString SendLabel(TEXT("send-datachannel"));
-			FUtf8String Utf8SendLabel = *SendLabel;
-			EpicRtcDataSource SendDataSource = {
-				._label = ToEpicRtcStringView(Utf8RecvLabel),
-				._maxRetransmitTime = 0,
-				._maxRetransmits = 0,
-				._isOrdered = true,
-				._protocol = EpicRtcDataSourceProtocol::Sctp,
-				._negotiated = true,
-				._transportChannelId = 0
-			};
-			ParticipantConnection->AddDataSource(SendDataSource);
-		}
-		else
-		{
-			EpicRtcDataSource DataSource = {
-				._label = Participant->GetId(),
-				._maxRetransmitTime = 0,
-				._maxRetransmits = 0,
-				._isOrdered = true,
-				._protocol = EpicRtcDataSourceProtocol::Sctp
-			};
-			ParticipantConnection->AddDataSource(DataSource);
-		}
-
-		EpicRtcBitrate Bitrates = {
-			._minBitrateBps = UPixelStreaming2PluginSettings::CVarWebRTCMinBitrate.GetValueOnAnyThread(),
-			._hasMinBitrateBps = true,
-			._maxBitrateBps = UPixelStreaming2PluginSettings::CVarWebRTCMaxBitrate.GetValueOnAnyThread(),
-			._hasMaxBitrateBps = true,
-			._startBitrateBps = UPixelStreaming2PluginSettings::CVarWebRTCStartBitrate.GetValueOnAnyThread(),
-			._hasStartBitrateBps = true
-		};
-
-		ParticipantConnection->SetConnectionRates(Bitrates);
-
-		ParticipantConnection->StartNegotiation();
-	}
-
-	void FStreamer::OnRoomLeftUpdate(const EpicRtcStringView Participant)
-	{
-		FString ParticipantId = ToString(Participant);
-		UE_LOG(LogPixelStreaming2, Log, TEXT("FStreamer::OnRoomLeftUpdate(Participant [%s] left the room.)"), *ParticipantId);
-
-		// Remove the player
-		DeletePlayerSession(ParticipantId);
-	}
-
-	void FStreamer::OnAudioTrackUpdate(EpicRtcParticipantInterface* Participant, EpicRtcAudioTrackInterface* AudioTrack)
-	{
-		FString ParticipantId = ToString(Participant->GetId());
-		FString AudioTrackId = ToString(AudioTrack->GetId());
-		UE_LOG(LogPixelStreaming2, Log, TEXT("FStreamer::OnAudioTrackUpdate(Participant [%s], AudioTrack [%s])"), *ParticipantId, *AudioTrackId);
-
-		if (FPlayerContext* PlayerContext = Players->Find(ParticipantId))
-		{
-			if (AudioTrack->IsRemote())
-			{
-				PlayerContext->AudioSink = FEpicRtcAudioSink::Create(AudioTrack);
-			}
-			else
-			{
-				PlayerContext->AudioSource = FEpicRtcAudioSource::Create(AudioTrack);
-			}
-
-			if (UPixelStreaming2Delegates* Delegates = UPixelStreaming2Delegates::Get())
-			{
-				Delegates->OnAudioTrackOpenNative.Broadcast(StreamerId, ParticipantId, (bool)AudioTrack->IsRemote());
-			}
-		}
-	}
-
-	void FStreamer::OnVideoTrackUpdate(EpicRtcParticipantInterface* Participant, EpicRtcVideoTrackInterface* VideoTrack)
-	{
-		FString ParticipantId = ToString(Participant->GetId());
-		FString VideoTrackId = ToString(VideoTrack->GetId());
-		UE_LOG(LogPixelStreaming2, Log, TEXT("FStreamer::OnVideoTrackUpdate(Participant [%s], VideoTrack [%s])"), *ParticipantId, *VideoTrackId);
-
-		if (FPlayerContext* PlayerContext = Players->Find(ParticipantId))
-		{
-			if (VideoTrack->IsRemote())
-			{
-				PlayerContext->VideoSink = FEpicRtcVideoSink::Create(VideoTrack);
-			}
-			else
-			{
-				PlayerContext->VideoSource = FEpicRtcVideoSource::Create(VideoTrack, VideoCapturer, VideoSourceGroup);
-			}
-
-			if (UPixelStreaming2Delegates* Delegates = UPixelStreaming2Delegates::Get())
-			{
-				Delegates->OnVideoTrackOpenNative.Broadcast(StreamerId, ParticipantId, (bool)VideoTrack->IsRemote());
-			}
-		}
-	}
-
-	void FStreamer::OnDataTrackUpdate(EpicRtcParticipantInterface* Participant, EpicRtcDataTrackInterface* DataTrack)
-	{
-		FString ParticipantId = ToString(Participant->GetId());
-		FString DataTrackId = ToString(DataTrack->GetId());
-		UE_LOG(LogPixelStreaming2, Log, TEXT("FStreamer::OnDataTrackUpdate(Participant [%s], DataTrack [%s])"), *ParticipantId, *DataTrackId);
-
-		if (FPlayerContext* PlayerContext = Players->Find(ParticipantId))
-		{
-			if (!PlayerContext->DataTrack)
-			{
-				PlayerContext->DataTrack = FEpicRtcDataTrack::Create(DataTrack, InputHandler->GetFromStreamerProtocol());
-			}
-			else
-			{
-				PlayerContext->DataTrack->SetSendTrack(DataTrack);
-			}
-		}
-	}
-
-	void FStreamer::OnLocalSdpUpdate(EpicRtcParticipantInterface* Participant, EpicRtcSdpInterface* Sdp)
-	{
-		FString ParticipantId = ToString(Participant->GetId());
-		FString SdpType = TEXT("");
-		switch (Sdp->GetType())
-		{
-			case EpicRtcSdpType::Offer:
-				SdpType = TEXT("Offer");
-				break;
-			case EpicRtcSdpType::Answer:
-				SdpType = TEXT("Answer");
-				break;
-		}
-		UE_LOG(LogPixelStreaming2, Log, TEXT("FStreamer::OnLocalSdpUpdate(Participant [%s], Type [%s])"), *ParticipantId, *SdpType);
-	}
-
-	void FStreamer::OnRemoteSdpUpdate(EpicRtcParticipantInterface* Participant, EpicRtcSdpInterface* Sdp)
-	{
-		FString ParticipantId = ToString(Participant->GetId());
-		FString SdpType = TEXT("");
-		switch (Sdp->GetType())
-		{
-			case EpicRtcSdpType::Offer:
-				SdpType = TEXT("Offer");
-				break;
-			case EpicRtcSdpType::Answer:
-				SdpType = TEXT("Answer");
-				break;
-		}
-		UE_LOG(LogPixelStreaming2, Log, TEXT("FStreamer::OnRemoteSdpUpdate(Participant [%s], Type [%s])"), *ParticipantId, *SdpType);
-	}
-
-	void FStreamer::OnRoomErrorUpdate(const EpicRtcErrorCode Error)
-	{
-		UE_LOGFMT(LogPixelStreaming2, VeryVerbose, "FStreamer::OnRoomErrorUpdate does nothing");
-	}
-
-	void FStreamer::OnAudioTrackMuted(EpicRtcAudioTrackInterface* AudioTrack, EpicRtcBool bIsMuted)
-	{
-		FString PlayerId;
-		bool	bFoundPlayer = FindPlayerByAudioTrack(AudioTrack, PlayerId);
-		FString AudioTrackId = ToString(AudioTrack->GetId());
-		UE_LOG(LogPixelStreaming2, Log, TEXT("FStreamer::OnAudioTrackMuted(AudioTrack [%s], bIsMuted[%s], PlayerId[%s])"), *AudioTrackId, bIsMuted ? TEXT("True") : TEXT("False"), *PlayerId);
-
-		if (!bFoundPlayer)
-		{
-			UE_LOG(LogPixelStreaming2, Warning, TEXT("FStreamer::OnAudioTrackMuted(Failed to find a player for audio track [%s])"), *AudioTrackId);
-			return;
-		}
-
-		FPlayerContext* PlayerContext = Players->Find(PlayerId);
-		if (AudioTrack->IsRemote())
-		{
-			PlayerContext->AudioSink->SetMuted((bool)bIsMuted);
-		}
-		else
-		{
-			PlayerContext->AudioSource->SetMuted((bool)bIsMuted);
-		}
-	}
-
-	void FStreamer::OnAudioTrackFrame(EpicRtcAudioTrackInterface* AudioTrack, const EpicRtcAudioFrame& Frame)
-	{
-		FString PlayerId;
-		bool	bFoundPlayer = FindPlayerByAudioTrack(AudioTrack, PlayerId);
-		FString AudioTrackId = ToString(AudioTrack->GetId());
-
-		if (!bFoundPlayer)
-		{
-			UE_LOG(LogPixelStreaming2, Warning, TEXT("FStreamer::OnAudioTrackFrame(Failed to find a player for audio track [%s])"), *AudioTrackId);
-			return;
-		}
-
-		FPlayerContext* PlayerContext = Players->Find(PlayerId);
-		if (PlayerContext->AudioSink)
-		{
-			PlayerContext->AudioSink->OnAudioData(Frame._data, Frame._length, Frame._format._numChannels, Frame._format._sampleRate);
-		}
-	}
-
-	void FStreamer::OnAudioTrackRemoved(EpicRtcAudioTrackInterface* AudioTrack)
-	{
-		FString PlayerId;
-		bool	bFoundPlayer = FindPlayerByAudioTrack(AudioTrack, PlayerId);
-		FString AudioTrackId = ToString(AudioTrack->GetId());
-		UE_LOG(LogPixelStreaming2, Log, TEXT("FStreamer::OnAudioTrackRemoved(AudioTrack [%s], PlayerId[%s])"), *AudioTrackId, *PlayerId);
-
-		if (!bFoundPlayer)
-		{
-			UE_LOG(LogPixelStreaming2, Warning, TEXT("FStreamer::OnAudioTrackFrame(Failed to find a player for audio track [%s])"), *AudioTrackId);
-			return;
-		}
-
-		FPlayerContext* PlayerContext = Players->Find(PlayerId);
-		if (AudioTrack->IsRemote())
-		{
-			PlayerContext->AudioSink = nullptr;
-		}
-		else
-		{
-			PlayerContext->AudioSource = nullptr;
-		}
-
-		if (UPixelStreaming2Delegates* Delegates = UPixelStreaming2Delegates::Get())
-		{
-			Delegates->OnAudioTrackClosedNative.Broadcast(StreamerId, PlayerId, (bool)AudioTrack->IsRemote());
-		}
-	}
-
-	void FStreamer::OnAudioTrackState(EpicRtcAudioTrackInterface* AudioTrack, const EpicRtcTrackState State)
-	{
-		UE_LOGFMT(LogPixelStreaming2, VeryVerbose, "FStreamer::OnAudioTrackState does nothing");
-	}
-
-	void FStreamer::OnVideoTrackMuted(EpicRtcVideoTrackInterface* VideoTrack, EpicRtcBool bIsMuted)
-	{
-		FString PlayerId;
-		bool	bFoundPlayer = FindPlayerByVideoTrack(VideoTrack, PlayerId);
-		FString VideoTrackId = ToString(VideoTrack->GetId());
-		UE_LOG(LogPixelStreaming2, Log, TEXT("FStreamer::OnVideoTrackMuted(VideoTrack [%s], bIsMuted[%s], PlayerId[%s])"), *VideoTrackId, bIsMuted ? TEXT("True") : TEXT("False"), *PlayerId);
-
-		if (!bFoundPlayer)
-		{
-			UE_LOG(LogPixelStreaming2, Warning, TEXT("FStreamer::OnVideoTrackMuted(Failed to find a player for video track [%s])"), *VideoTrackId);
-			return;
-		}
-
-		FPlayerContext* PlayerContext = Players->Find(PlayerId);
-		if (VideoTrack->IsRemote())
-		{
-			PlayerContext->VideoSink->SetMuted((bool)bIsMuted);
-		}
-		else
-		{
-			PlayerContext->VideoSource->SetMuted((bool)bIsMuted);
-		}
-	}
-
-	void FStreamer::OnVideoTrackFrame(EpicRtcVideoTrackInterface* VideoTrack, const EpicRtcVideoFrame& Frame)
-	{
-		FString PlayerId;
-		bool	bFoundPlayer = FindPlayerByVideoTrack(VideoTrack, PlayerId);
-		FString VideoTrackId = ToString(VideoTrack->GetId());
-
-		if (!bFoundPlayer)
-		{
-			UE_LOG(LogPixelStreaming2, Warning, TEXT("FStreamer::OnVideoTrackFrame(Failed to find a player for video track [%s])"), *VideoTrackId);
-			return;
-		}
-
-		FPlayerContext* PlayerContext = Players->Find(PlayerId);
-		if (PlayerContext->VideoSink)
-		{
-			PlayerContext->VideoSink->OnVideoData(Frame);
-		}
-	}
-
-	void FStreamer::OnVideoTrackRemoved(EpicRtcVideoTrackInterface* VideoTrack)
-	{
-		FString PlayerId;
-		bool	bFoundPlayer = FindPlayerByVideoTrack(VideoTrack, PlayerId);
-		FString VideoTrackId = ToString(VideoTrack->GetId());
-
-		if (!bFoundPlayer)
-		{
-			UE_LOG(LogPixelStreaming2, Warning, TEXT("FStreamer::OnVideoTrackRemoved(Failed to find a player for video track [%s])"), *VideoTrackId);
-			return;
-		}
-
-		UE_LOG(LogPixelStreaming2, Log, TEXT("FStreamer::OnVideoTrackRemoved(VideoTrack=[%s], Player=[%s])"), *VideoTrackId, *PlayerId);
-
-		// If we did find a player, clear its video sink/source
-		FPlayerContext* PlayerContext = Players->Find(PlayerId);
-		if (VideoTrack->IsRemote())
-		{
-			PlayerContext->VideoSink = nullptr;
-		}
-		else
-		{
-			PlayerContext->VideoSource = nullptr;
-		}
-
-		if (UPixelStreaming2Delegates* Delegates = UPixelStreaming2Delegates::Get())
-		{
-			Delegates->OnVideoTrackClosedNative.Broadcast(StreamerId, PlayerId, (bool)VideoTrack->IsRemote());
-		}
-	}
-
-	void FStreamer::OnVideoTrackState(EpicRtcVideoTrackInterface* VideoTrack, const EpicRtcTrackState State)
-	{
-		FString PlayerId;
-		bool	bFoundPlayer = FindPlayerByVideoTrack(VideoTrack, PlayerId);
-		FString VideoTrackId = ToString(VideoTrack->GetId());
-
-		// Note: It is acceptable to not have a found a player for track state changes, as these can trigger before we have added a participant
-
-		if (State == EpicRtcTrackState::Active)
-		{
-			UE_LOG(LogPixelStreaming2, Log, TEXT("FStreamer::OnVideoTrackState(VideoTrack=[%s], Player=[%s], State=Active)"), *VideoTrackId, *PlayerId);
-		}
-		else if (State == EpicRtcTrackState::Stopped)
-		{
-			UE_LOG(LogPixelStreaming2, Log, TEXT("FStreamer::OnVideoTrackState(VideoTrack=[%s], Player=[%s], State=Stopped)"), *VideoTrackId, *PlayerId);
-		}
-	}
-
-	void FStreamer::OnDataTrackRemoved(EpicRtcDataTrackInterface* DataTrack)
-	{
-		// As long as DataTrack emits `stopped` state when it removed, this should be enough.
-		// because OnDataTrackState(DataTrack, EpicRtcTrackState::Stopped) already calls `OnDataChannelClosed`.
-	}
-
-	void FStreamer::OnDataTrackState(EpicRtcDataTrackInterface* DataTrack, const EpicRtcTrackState State)
-	{
-		FString PlayerId;
-		bool	bFoundPlayer = FindPlayerByDataTrack(DataTrack, PlayerId);
-		FString DataTrackId = ToString(DataTrack->GetId());
-
-		if (State == EpicRtcTrackState::Active)
-		{
-			if (!bFoundPlayer)
-			{
-				UE_LOG(LogPixelStreaming2, Log, TEXT("FStreamer::OnDataTrackState(Message was State=Active. Failed to find a player for data track [%s])"), *DataTrackId);
-				return;
-			}
-			UE_LOG(LogPixelStreaming2, Log, TEXT("FStreamer::OnDataTrackState(Message was State=Active. Player [%s])"), *DataTrackId);
-			OnDataChannelOpen(PlayerId);
-		}
-		else if (State == EpicRtcTrackState::Stopped)
-		{
-			UE_LOG(LogPixelStreaming2, Log, TEXT("FStreamer::OnDataTrackState(Message was State=Stopped. Player [%s])"), *DataTrackId);
-			OnDataChannelClosed(PlayerId);
-		}
-	}
-
-	void FStreamer::OnDataTrackMessage(EpicRtcDataTrackInterface* DataTrack)
-	{
-		FString DataTrackId = ToString(DataTrack->GetId());
-		TRefCountPtr<EpicRtcDataFrameInterface> DataFrame;
-		if (!DataTrack->PopFrame(DataFrame.GetInitReference()))
-		{
-			UE_LOG(LogPixelStreaming2, Log, TEXT("FStreamer::OnDataTrackMessage(Failed to PopFrame [%s])"), *DataTrackId);
-			return;
-		}
-		FString PlayerId;
-		const uint8_t* Data = DataFrame->Data();
-		uint32_t DataSize = DataFrame->Size();
-		uint8 Type = Data[0];
-		TSharedPtr<IPixelStreaming2DataProtocol> ToStreamerProtocol = InputHandler->GetToStreamerProtocol();
-		if (Type == ToStreamerProtocol->Find(EPixelStreaming2ToStreamerMessage::Multiplexed)->GetID())
-		{
-			//skip type
-			Data++;
-			DataSize--;
-			PlayerId = ReadString(Data, DataSize); 
-			Type = Data[0];
-			UE_LOG(LogPixelStreaming2, VeryVerbose, TEXT("FStreamer::OnDataTrackMessage(Received multiplexed message of type [%d] with PlayerId [%s])"), Type, *PlayerId);
-		}
-		else if (Type == ToStreamerProtocol->Find(EPixelStreaming2ToStreamerMessage::ChannelRelayStatus)->GetID())
-		{
-			HandleRelayStatusMessage(Data, DataSize, DataTrack);
-			return;
-		}
-		else if (!FindPlayerByDataTrack(DataTrack, PlayerId))
-		{
-			UE_LOG(LogPixelStreaming2, Log, TEXT("FStreamer::OnDataTrackMessage(Failed to find a player for data track [%s])"), *DataTrackId);
-			return;
-		}
-		
-		if (Type == ToStreamerProtocol->Find(EPixelStreaming2ToStreamerMessage::LatencyTest)->GetID())
-		{
-			SendLatencyReport(PlayerId);
-		}
-		else if (Type == ToStreamerProtocol->Find(EPixelStreaming2ToStreamerMessage::RequestInitialSettings)->GetID())
-		{
-			SendInitialSettings(PlayerId);
-		}
-		else if (Type == ToStreamerProtocol->Find(EPixelStreaming2ToStreamerMessage::IFrameRequest)->GetID())
-		{
-			ForceKeyFrame();
-		}
-		else if (Type == ToStreamerProtocol->Find(EPixelStreaming2ToStreamerMessage::TestEcho)->GetID())
-		{
-			if (FPlayerContext* PlayerContext = Players->Find(PlayerId))
-			{
-				if (PlayerContext->DataTrack)
-				{
-					const size_t  DescriptorSize = (DataSize - 1) / sizeof(TCHAR);
-					const TCHAR*  DescPtr = reinterpret_cast<const TCHAR*>(Data + 1);
-					const FString Message(DescriptorSize, DescPtr);
-					PlayerContext->DataTrack->SendMessage(EPixelStreaming2FromStreamerMessage::TestEcho, Message);
-				}
-			}
-		}
-		else if (!IsEngineExitRequested())
-		{
-			// If we are in "Host" mode and the current peer is not the host, then discard this input.
-			if (GetEnumFromCVar<EInputControllerMode>(UPixelStreaming2PluginSettings::CVarInputController) == EInputControllerMode::Host && InputControllingId != PlayerId)
-			{
-				return;
-			}
-			
-			TArray<uint8> MessageData(Data, DataSize);
-			if (InputHandler)
-			{
-				InputHandler->OnMessage(PlayerId, MessageData);
-			}
-		}
-	}
-
-	bool FStreamer::FindPlayerByAudioTrack(EpicRtcAudioTrackInterface* AudioTrack, FString& OutPlayerId)
-	{
-		OutPlayerId = "";
-		// Find the player corresponding to this data track
-		FString AudioTrackId = ToString(AudioTrack->GetId());
-		Players->ApplyUntil([this, AudioTrackId, &OutPlayerId](FString DataPlayerId, FPlayerContext& PlayerContext) {
-			if (PlayerContext.AudioSource)
-			{
-				FString TrackId = ToString(PlayerContext.AudioSource->GetTrackId());
-				if (TrackId == AudioTrackId)
-				{
-					OutPlayerId = DataPlayerId;
-					return true;
-				}
-			}
-			if (PlayerContext.AudioSink)
-			{
-				FString TrackId = ToString(PlayerContext.AudioSink->GetTrackId());
-				if (TrackId == AudioTrackId)
-				{
-					OutPlayerId = DataPlayerId;
-					return true;
-				}
-			}
-			return false;
-		});
-
-		return !OutPlayerId.IsEmpty();
-	}
-
-	bool FStreamer::FindPlayerByVideoTrack(EpicRtcVideoTrackInterface* VideoTrack, FString& OutPlayerId)
-	{
-		OutPlayerId = "";
-		// Find the player corresponding to this data track
-		FString VideoTrackId = ToString(VideoTrack->GetId());
-		Players->ApplyUntil([this, VideoTrackId, &OutPlayerId](FString DataPlayerId, FPlayerContext& PlayerContext) {
-			if (PlayerContext.VideoSource)
-			{
-				FString TrackId = ToString(PlayerContext.VideoSource->GetTrackId());
-				if (TrackId == VideoTrackId)
-				{
-					OutPlayerId = DataPlayerId;
-					return true;
-				}
-			}
-			if (PlayerContext.VideoSink)
-			{
-				FString TrackId = ToString(PlayerContext.VideoSink->GetTrackId());
-				if (TrackId == VideoTrackId)
-				{
-					OutPlayerId = DataPlayerId;
-					return true;
-				}
-			}
-			return false;
-		});
-
-		return !OutPlayerId.IsEmpty();
-	}
-
-	bool FStreamer::FindPlayerByDataTrack(EpicRtcDataTrackInterface* DataTrack, FString& OutPlayerId)
-	{
-		OutPlayerId = "";
-		// Find the player corresponding to this data track
-		FString DataTrackId = ToString(DataTrack->GetId());
-		Players->ApplyUntil([this, DataTrackId, &OutPlayerId](FString DataPlayerId, FPlayerContext& PlayerContext) {
-			if (PlayerContext.DataTrack)
-			{
-				FString TrackId = ToString(PlayerContext.DataTrack->GetId());
-				if (TrackId == DataTrackId)
-				{
-					OutPlayerId = DataPlayerId;
-					return true;
-				}
-			}
-			return false;
-		});
-
-		return !OutPlayerId.IsEmpty();
-	}
-
-} // namespace UE::PixelStreaming2
+	return StreamProtocols;
+}
\ No newline at end of file
diff --git a/Engine/Plugins/Media/PixelStreaming2/Source/PixelStreaming2/Private/Streamer.h b/Engine/Plugins/Media/PixelStreaming2/Source/PixelStreaming2/Private/Streamer.h
--- a/Engine/Plugins/Media/PixelStreaming2/Source/PixelStreaming2/Private/Streamer.h
+++ b/Engine/Plugins/Media/PixelStreaming2/Source/PixelStreaming2/Private/Streamer.h
@@ -1,161 +1,13 @@
 // Copyright Epic Games, Inc. All Rights Reserved.
 
-#pragma once
-
 #include "IPixelStreaming2Streamer.h"
-#include "VideoSourceGroup.h"
-#include "ThreadSafeMap.h"
-#include "Dom/JsonObject.h"
-#include "IPixelStreaming2InputHandler.h"
-#include "Templates/SharedPointer.h"
-#include "PlayerContext.h"
-#include "FreezeFrame.h"
-#include "StreamerReconnectTimer.h"
-#include "EpicRtcManager.h"
-
-#include "epic_rtc/core/stats.h"
-
-class IPixelStreaming2Module;
 
 namespace UE::PixelStreaming2
 {
-	static const FString INVALID_PLAYER_ID = FString(TEXT("Invalid Player Id"));
-
-	class FStreamer : public IPixelStreaming2Streamer, public TSharedFromThis<FStreamer>
+	class FDummyStreamerFactory : public IPixelStreaming2StreamerFactory
 	{
 	public:
-		static TSharedPtr<FStreamer> Create(const FString& StreamerId, TRefCountPtr<EpicRtcConferenceInterface> Conference);
-		virtual ~FStreamer();
-
-		virtual void  SetStreamFPS(int32 InFramesPerSecond) override;
-		virtual int32 GetStreamFPS() override;
-		virtual void  SetCoupleFramerate(bool bCouple) override;
-
-		virtual void									SetVideoProducer(TSharedPtr<IPixelStreaming2VideoProducer> Input) override;
-		virtual TWeakPtr<IPixelStreaming2VideoProducer> GetVideoProducer() override;
-
-		virtual void	SetSignallingServerURL(const FString& InSignallingServerURL) override;
-		virtual FString GetSignallingServerURL() override;
-
-		virtual FString GetId() override { return StreamerId; };
-		virtual bool	IsSignallingConnected() override { return bSignallingConnected; }
-		virtual void	StartStreaming() override;
-		virtual void	StopStreaming() override;
-		virtual bool	IsStreaming() const override { return bStreamingStarted; }
-
-		virtual FPreConnectionEvent&	OnPreConnection() override;
-		virtual FStreamingStartedEvent& OnStreamingStarted() override;
-		virtual FStreamingStoppedEvent& OnStreamingStopped() override;
-
-		virtual void ForceKeyFrame() override;
-
-		virtual void FreezeStream(UTexture2D* Texture) override;
-		virtual void UnfreezeStream() override;
-
-		virtual void			SendAllPlayersMessage(FString MessageType, const FString& Descriptor) override;
-		virtual void			SendPlayerMessage(FString PlayerId, FString MessageType, const FString& Descriptor) override;
-		virtual void			SendFileData(const TArray64<uint8>& ByteData, FString& MimeType, FString& FileExtension) override;
-		virtual void			KickPlayer(FString PlayerId) override;
-		virtual TArray<FString> GetConnectedPlayers() override;
-
-		virtual TWeakPtr<IPixelStreaming2InputHandler> GetInputHandler() override { return InputHandler; }
-
-		virtual IPixelStreaming2AudioSink* GetPeerAudioSink(FString PlayerId) override;
-		virtual IPixelStreaming2AudioSink* GetUnlistenedAudioSink() override;
-		virtual IPixelStreaming2VideoSink* GetPeerVideoSink(FString PlayerId) override;
-		virtual IPixelStreaming2VideoSink* GetUnwatchedVideoSink() override;
-
-		virtual void SetConfigOption(const FName& OptionName, const FString& Value) override;
-		virtual bool GetConfigOption(const FName& OptionName, FString& OutValue) override;
-
-		virtual void PlayerRequestsBitrate(FString PlayerId, int MinBitrate, int MaxBitrate) override;
-
-		virtual void RefreshStreamBitrate() override;
-
-		void ForEachPlayer(const TFunction<void(FString, FPlayerContext)>& Func);
-
-	private:
-		FStreamer(const FString& StreamerId, TRefCountPtr<EpicRtcConferenceInterface> Conference);
-
-		// own methods
-		void OnProtocolUpdated();
-		void ConsumeStats(FString PlayerId, FName StatName, float StatValue);
-		void DeletePlayerSession(FString PlayerId);
-		void DeleteAllPlayerSessions();
-		void OnDataChannelOpen(FString PlayerId);
-		void OnDataChannelClosed(FString PlayerId);
-		void SendInitialSettings(FString PlayerId) const;
-		void SendProtocol(FString PlayerId) const;
-		void SendPeerControllerMessages(FString PlayerId) const;
-		void SendLatencyReport(FString PlayerId) const;
-		void HandleRelayStatusMessage(const uint8_t* Data, uint32_t Size, EpicRtcDataTrackInterface* DataTrack);
-		bool FindPlayerByAudioTrack(EpicRtcAudioTrackInterface* AudioTrack, FString& OutPlayerId);
-		bool FindPlayerByVideoTrack(EpicRtcVideoTrackInterface* VideoTrack, FString& OutPlayerId);
-		bool FindPlayerByDataTrack(EpicRtcDataTrackInterface* DataTrack, FString& OutPlayerId);
-		void TriggerMouseLeave(FString InStreamerId);
-		void RemoveSession(bool bDisconnect);
-		void RemoveRoom();
-		void OnStatsReady(const FString& PlayerId, const EpicRtcConnectionStats& ConnectionStats);
-
-	private:
-		FString StreamerId;
-		FString CurrentSignallingServerURL;
-
-		TSharedPtr<IPixelStreaming2InputHandler> InputHandler;
-
-		TSharedPtr<TThreadSafeMap<FString, FPlayerContext>> Players;
-
-		FString InputControllingId = INVALID_PLAYER_ID;
-
-		bool bSignallingConnected = false;
-		bool bStreamingStarted = false;
-
-		FPreConnectionEvent	   StreamingPreConnectionEvent;
-		FStreamingStartedEvent StreamingStartedEvent;
-		FStreamingStoppedEvent StreamingStoppedEvent;
-
-		TSharedPtr<FVideoCapturer>	  VideoCapturer;
-		TSharedPtr<FVideoSourceGroup> VideoSourceGroup;
-		TSharedPtr<FFreezeFrame>	  FreezeFrame;
-
-		FDelegateHandle ConsumeStatsHandle;
-		FDelegateHandle AllConnectionsClosedHandle;
-
-		TMap<FName, FString> ConfigOptions;
-
-		TSharedPtr<FStreamerReconnectTimer> ReconnectTimer;
-
-		TSharedPtr<FEpicRtcManager> EpicRtcManager;
-
-	public:
-		// Begin FEpicRtcManager Callbacks
-		void OnSessionStateUpdate(const EpicRtcSessionState StateUpdate);
-		void OnSessionErrorUpdate(const EpicRtcErrorCode ErrorUpdate);
-		void OnSessionRoomsAvailableUpdate(EpicRtcStringArrayInterface* RoomsList);
-
-		void OnRoomStateUpdate(const EpicRtcRoomState State);
-		void OnRoomJoinedUpdate(EpicRtcParticipantInterface* Participant);
-		void OnRoomLeftUpdate(const EpicRtcStringView ParticipantId);
-		void OnAudioTrackUpdate(EpicRtcParticipantInterface* Participant, EpicRtcAudioTrackInterface* AudioTrack);
-		void OnVideoTrackUpdate(EpicRtcParticipantInterface* Participant, EpicRtcVideoTrackInterface* VideoTrack);
-		void OnDataTrackUpdate(EpicRtcParticipantInterface* Participant, EpicRtcDataTrackInterface* DataTrack);
-		void OnLocalSdpUpdate(EpicRtcParticipantInterface* Participant, EpicRtcSdpInterface* Sdp);
-		void OnRemoteSdpUpdate(EpicRtcParticipantInterface* Participant, EpicRtcSdpInterface* Sdp);
-		void OnRoomErrorUpdate(const EpicRtcErrorCode Error);
-
-		void OnAudioTrackMuted(EpicRtcAudioTrackInterface* AudioTrack, EpicRtcBool bIsMuted);
-		void OnAudioTrackFrame(EpicRtcAudioTrackInterface* AudioTrack, const EpicRtcAudioFrame& Frame);
-		void OnAudioTrackRemoved(EpicRtcAudioTrackInterface* AudioTrack);
-		void OnAudioTrackState(EpicRtcAudioTrackInterface* AudioTrack, const EpicRtcTrackState State);
-
-		void OnVideoTrackMuted(EpicRtcVideoTrackInterface* VideoTrack, EpicRtcBool bIsMuted);
-		void OnVideoTrackFrame(EpicRtcVideoTrackInterface* VideoTrack, const EpicRtcVideoFrame& Frame);
-		void OnVideoTrackRemoved(EpicRtcVideoTrackInterface* VideoTrack);
-		void OnVideoTrackState(EpicRtcVideoTrackInterface* VideoTrack, const EpicRtcTrackState State);
-
-		void OnDataTrackRemoved(EpicRtcDataTrackInterface* DataTrack);
-		void OnDataTrackState(EpicRtcDataTrackInterface* DataTrack, const EpicRtcTrackState State);
-		void OnDataTrackMessage(EpicRtcDataTrackInterface* DataTrack);
-		// End FEpicRtcManager Callbacks
+		virtual FString								 GetStreamType() override;
+		virtual TSharedPtr<IPixelStreaming2Streamer> CreateNewStreamer(const FString& StreamerId) override;
 	};
-} // namespace UE::PixelStreaming2
+} // namespace UE::PixelStreaming2
\ No newline at end of file
diff --git a/Engine/Plugins/Media/PixelStreaming2/Source/PixelStreaming2/Private/Thread.cpp b/Engine/Plugins/Media/PixelStreaming2/Source/PixelStreaming2/Private/Thread.cpp
new file mode 100644
--- /dev/null
+++ b/Engine/Plugins/Media/PixelStreaming2/Source/PixelStreaming2/Private/Thread.cpp
@@ -0,0 +1,235 @@
+// Copyright Epic Games, Inc. All Rights Reserved.
+
+#include "Thread.h"
+
+#include "HAL/Event.h"
+#include "HAL/Runnable.h"
+#include "HAL/PlatformProcess.h"
+#include "HAL/PlatformTime.h"
+#include "HAL/RunnableThread.h"
+#include "Misc/SingleThreadRunnable.h"
+#include "TickableTask.h"
+
+namespace UE::PixelStreaming2
+{
+	static TWeakPtr<FPixelStreamingRunnable> PixelStreamingRunnable;
+
+	/**
+	 * The runnable. Handles ticking of all tasks
+	 */
+	class FPixelStreamingRunnable : public FRunnable, public FSingleThreadRunnable
+	{
+	public:
+		// Begin FRunnable
+		virtual bool Init() override
+		{
+			return true;
+		}
+
+		virtual uint32 Run() override
+		{
+			bIsRunning = true;
+
+			while (bIsRunning)
+			{
+				Tick();
+
+				// Sleep 1ms
+				FPlatformProcess::Sleep(0.001f);
+			}
+
+			return 0;
+		}
+
+		virtual void Stop() override
+		{
+			bIsRunning = false;
+
+			TaskEvent->Trigger();
+		}
+
+		virtual void Exit() override
+		{
+			bIsRunning = false;
+
+			TaskEvent->Trigger();
+		}
+
+		virtual FSingleThreadRunnable* GetSingleThreadInterface() override
+		{
+			return this;
+		}
+		// End FRunnable
+
+		// Begin FSingleThreadRunnable
+		virtual void Tick() override
+		{
+			FScopeLock TaskLock(&TasksMutex);
+
+			const uint64 NowCycles = FPlatformTime::Cycles64();
+			const double DeltaMs = FPlatformTime::ToMilliseconds64(NowCycles - LastTickCycles);
+
+			StartTicking();
+
+			for (auto& Task : Tasks)
+			{
+				// A task may be nulled out due to deletion during our loop. Check for safety
+				if (Task)
+				{
+					Task->Tick(DeltaMs);
+				}
+			}
+
+			FinishTicking();
+
+			LastTickCycles = NowCycles;
+		}
+		// End FSingleThreadRunnable
+
+		void StartTicking()
+		{
+			FScopeLock NewTaskLock(&NewTasksMutex);
+			for (auto& NewTask : NewTasks)
+			{
+				Tasks.Add(NewTask);
+			}
+
+			NewTasks.Empty();
+
+			bIsTicking = true;
+		}
+
+		void FinishTicking()
+		{
+			if (bNeedsCleanup)
+			{
+				Tasks.RemoveAll([](const FPixelStreamingTickableTask* Entry) { return Entry == nullptr; });
+				bNeedsCleanup = false;
+			}
+
+			bIsTicking = false;
+
+			if (Tasks.Num() == 0)
+			{
+				// Sleep the thread indefinitely because there are no tasks to tick.
+				// Adding a new task will wake the thread
+				TaskEvent->Wait();
+			}
+		}
+
+		FPixelStreamingRunnable()
+			: bIsTicking(false)
+			, bNeedsCleanup(false)
+			, bIsRunning(false)
+			, LastTickCycles(FPlatformTime::Cycles64())
+		{
+		}
+
+		virtual ~FPixelStreamingRunnable() = default;
+
+	private:
+		void AddTask(FPixelStreamingTickableTask* Task)
+		{
+			FScopeLock NewTaskLock(&NewTasksMutex);
+			NewTasks.Add(Task);
+			// We've added a new task. Wake the thread (if it was sleeping)
+			TaskEvent->Trigger();
+		}
+
+		void RemoveTask(FPixelStreamingTickableTask* Task)
+		{
+			// Lock TaskLock before NewTaskLock to ensure deadlock does not happen when Tick and StartTicking lock.
+			// Locking matches FTickableObjectBase locking in Tickable.cpp
+			FScopeLock TaskLock(&TasksMutex);
+			FScopeLock NewTaskLock(&NewTasksMutex);
+
+			// Remove from pending list if it hasn't been registered
+			NewTasks.Remove(Task);
+
+			const int32 Pos = Tasks.IndexOfByKey(Task);
+			if (Pos != INDEX_NONE)
+			{
+				if (bIsTicking)
+				{
+					// During ticking it is not safe to modify the set so null and mark for later
+					Tasks[Pos] = nullptr;
+					bNeedsCleanup = true;
+				}
+				else
+				{
+					Tasks.RemoveAt(Pos);
+				}
+			}
+		}
+
+	private:
+		// Allow the FPixelStreamingTickableTask to access the private add and remove tasks
+		friend FPixelStreamingTickableTask;
+
+		// New tasks that have not yet been added to the Tasks list*/
+		TArray<FPixelStreamingTickableTask*> NewTasks;
+		// Lock for modifying new list */
+		FCriticalSection NewTasksMutex;
+
+		// Tasks to execute every tick
+		TArray<FPixelStreamingTickableTask*> Tasks;
+		// This critical section should be locked during entire tick process
+		FCriticalSection TasksMutex;
+
+		// Use this event to signal when we should wake.
+		FEventRef TaskEvent;
+
+		// Tasks can removed from any thread so this needs to be thread safe
+		std::atomic<bool> bIsTicking;
+		// Tasks can removed from any thread so this needs to be thread safe
+		std::atomic<bool> bNeedsCleanup;
+		// This thread can be stopped from another thread during shutdown so this needs to be thread safe
+		std::atomic<bool> bIsRunning;
+		uint64			  LastTickCycles;
+	};
+
+	FPixelStreamingThread::FPixelStreamingThread()
+	{
+		Runnable = MakeShared<FPixelStreamingRunnable>();
+		PixelStreamingRunnable = Runnable;
+		Thread = TSharedPtr<FRunnableThread>(FRunnableThread::Create(Runnable.Get(), TEXT("Pixel Streaming PixelStreaming Thread")));
+	}
+
+	FPixelStreamingThread::~FPixelStreamingThread()
+	{
+		if (Thread)
+		{
+			Thread->Kill();
+			Thread.Reset();
+		}
+
+		if (Runnable)
+		{
+			Runnable->Stop();
+			Runnable.Reset();
+		}
+	}
+
+	/**
+	 * ---------- FPixelStreamingTickableTask ---------------
+	 */
+
+	void FPixelStreamingTickableTask::Register()
+	{
+		TSharedPtr<FPixelStreamingRunnable> Runnable = PixelStreamingRunnable.Pin();
+		if (Runnable)
+		{
+			Runnable->AddTask(this);
+		}
+	}
+
+	void FPixelStreamingTickableTask::Unregister()
+	{
+		TSharedPtr<FPixelStreamingRunnable> Runnable = PixelStreamingRunnable.Pin();
+		if (Runnable)
+		{
+			Runnable->RemoveTask(this);
+		}
+	}
+
+} // namespace UE::PixelStreaming2
\ No newline at end of file
diff --git a/Engine/Plugins/Media/PixelStreaming2/Source/PixelStreaming2/Private/Thread.h b/Engine/Plugins/Media/PixelStreaming2/Source/PixelStreaming2/Private/Thread.h
new file mode 100644
--- /dev/null
+++ b/Engine/Plugins/Media/PixelStreaming2/Source/PixelStreaming2/Private/Thread.h
@@ -0,0 +1,22 @@
+// Copyright Epic Games, Inc. All Rights Reserved.
+
+#pragma once
+
+#include "CoreMinimal.h"
+
+namespace UE::PixelStreaming2
+{
+	/**
+	 * The thread. Wraps both the RunnableThread and Runnable into a single point
+	 */
+	class FPixelStreamingThread
+	{
+	public:
+		FPixelStreamingThread();
+		~FPixelStreamingThread();
+
+	private:
+		TSharedPtr<class FRunnableThread>		  Thread;
+		TSharedPtr<class FPixelStreamingRunnable> Runnable;
+	};
+} // namespace UE::PixelStreaming2
diff --git a/Engine/Plugins/Media/PixelStreaming2/Source/PixelStreaming2/Private/VideoCapturer.cpp b/Engine/Plugins/Media/PixelStreaming2/Source/PixelStreaming2/Private/VideoCapturer.cpp
--- a/Engine/Plugins/Media/PixelStreaming2/Source/PixelStreaming2/Private/VideoCapturer.cpp
+++ b/Engine/Plugins/Media/PixelStreaming2/Source/PixelStreaming2/Private/VideoCapturer.cpp
@@ -2,7 +2,6 @@
 
 #include "VideoCapturer.h"
 
-#include "EpicRtcVideoBufferMultiFormat.h"
 #include "PixelCaptureBufferFormat.h"
 #include "PixelCaptureCapturerMediaCapture.h"
 #include "PixelCaptureCapturerI420.h"
@@ -90,11 +89,6 @@ namespace UE::PixelStreaming2
 		FrameCapturer->Capture(InputFrame);
 	}
 
-	TRefCountPtr<EpicRtcVideoBufferInterface> FVideoCapturer::GetFrameBuffer()
-	{
-		return new UE::PixelStreaming2::FEpicRtcVideoBufferMultiFormatLayered(FrameCapturer);
-	}
-
 	TSharedPtr<IPixelCaptureOutputFrame> FVideoCapturer::RequestFormat(int32 Format, int32 LayerIndex)
 	{
 		if (FrameCapturer != nullptr)
diff --git a/Engine/Plugins/Media/PixelStreaming2/Source/PixelStreaming2/Private/VideoProducerBackBuffer.cpp b/Engine/Plugins/Media/PixelStreaming2/Source/PixelStreaming2/Private/VideoProducerBackBuffer.cpp
--- a/Engine/Plugins/Media/PixelStreaming2/Source/PixelStreaming2/Private/VideoProducerBackBuffer.cpp
+++ b/Engine/Plugins/Media/PixelStreaming2/Source/PixelStreaming2/Private/VideoProducerBackBuffer.cpp
@@ -7,7 +7,7 @@
 #include "PixelCaptureBufferFormat.h"
 #include "PixelCaptureInputFrameRHI.h"
 #include "RHI.h"
-#include "UtilsCommon.h"
+#include "UtilsAsync.h"
 #include "Widgets/SWindow.h"
 
 
diff --git a/Engine/Plugins/Media/PixelStreaming2/Source/PixelStreaming2/Private/VideoProducerPIEViewport.cpp b/Engine/Plugins/Media/PixelStreaming2/Source/PixelStreaming2/Private/VideoProducerPIEViewport.cpp
--- a/Engine/Plugins/Media/PixelStreaming2/Source/PixelStreaming2/Private/VideoProducerPIEViewport.cpp
+++ b/Engine/Plugins/Media/PixelStreaming2/Source/PixelStreaming2/Private/VideoProducerPIEViewport.cpp
@@ -6,7 +6,7 @@
 #include "PixelCaptureInputFrameRHI.h"
 #include "RenderingThread.h"
 #include "UnrealClient.h"
-#include "UtilsCommon.h"
+#include "UtilsAsync.h"
 
 namespace UE::PixelStreaming2
 {
diff --git a/Engine/Plugins/Media/PixelStreaming2/Source/PixelStreaming2/Private/VideoSink.cpp b/Engine/Plugins/Media/PixelStreaming2/Source/PixelStreaming2/Private/VideoSink.cpp
--- a/Engine/Plugins/Media/PixelStreaming2/Source/PixelStreaming2/Private/VideoSink.cpp
+++ b/Engine/Plugins/Media/PixelStreaming2/Source/PixelStreaming2/Private/VideoSink.cpp
@@ -3,9 +3,6 @@
 #include "VideoSink.h"
 
 #include "Async/Async.h"
-#include "ColorConversion.h"
-#include "EpicRtcVideoBufferI420.h"
-#include "EpicRtcVideoBufferRHI.h"
 #include "IPixelStreaming2VideoConsumer.h"
 #include "PixelStreaming2Trace.h"
 #include "RenderTargetPool.h"
@@ -49,113 +46,6 @@ namespace UE::PixelStreaming2
 		return VideoConsumers.Num() > 0;
 	}
 
-	void FVideoSink::OnVideoData(const EpicRtcVideoFrame& Frame)
-	{
-		if (!HasVideoConsumers() || bIsMuted)
-		{
-			return;
-		}
-
-		TRACE_CPUPROFILER_EVENT_SCOPE_ON_CHANNEL_STR("FVideoSink::OnVideoData", PixelStreaming2Channel);
-
-		int32_t Width = Frame._buffer->GetWidth();
-		int32_t Height = Frame._buffer->GetHeight();
-
-		if (Frame._buffer->GetFormat() == EpicRtcPixelFormat::Native)
-		{
-			UE::PixelStreaming2::FEpicRtcVideoBufferRHI* const FrameBuffer = static_cast<UE::PixelStreaming2::FEpicRtcVideoBufferRHI*>(Frame._buffer);
-			if (FrameBuffer == nullptr)
-			{
-				return;
-			}
-
-			TSharedPtr<FVideoResourceRHI, ESPMode::ThreadSafe> VideoResource = FrameBuffer->GetVideoResource();
-			if (VideoResource->GetFormat() != EVideoFormat::BGRA)
-			{
-				VideoResource = VideoResource->TransformResource(FVideoDescriptor(EVideoFormat::BGRA, Width, Height));
-			}
-
-			auto& Raw = StaticCastSharedPtr<FVideoResourceRHI>(VideoResource)->GetRaw();
-			CallConsumeFrame(Raw.Texture);
-		}
-		else if (Frame._buffer->GetFormat() == EpicRtcPixelFormat::I420)
-		{
-			{
-				FScopeLock	  Lock(&RenderSyncContext);
-				const int32_t Size = Width * Height * 4;
-				if (Size > Buffer.Num())
-				{
-					Buffer.SetNum(Size);
-				}
-
-				UE::PixelStreaming2::FEpicRtcVideoBufferI420* const FrameBuffer = static_cast<UE::PixelStreaming2::FEpicRtcVideoBufferI420*>(Frame._buffer);
-
-				int StrideY = FrameBuffer->GetWidth();
-				int StrideUV = (FrameBuffer->GetWidth() + 1) / 2;
-
-				int DataSizeY = StrideY * FrameBuffer->GetHeight();
-				int DataSizeUV = StrideUV * ((FrameBuffer->GetHeight() + 1) / 2);
-
-				uint8_t* DataY = static_cast<uint8_t*>(FrameBuffer->GetData());
-				uint8_t* DataU = static_cast<uint8_t*>(FrameBuffer->GetData()) + DataSizeY;
-				uint8_t* DataV = static_cast<uint8_t*>(FrameBuffer->GetData()) + DataSizeY + DataSizeUV;
-
-				UE::PixelStreaming2::ConvertI420ToArgb(
-					DataY, StrideY,
-					DataU, StrideUV,
-					DataV, StrideUV,
-					Buffer.GetData(), Width * 4,
-					FrameBuffer->GetWidth(), FrameBuffer->GetHeight());
-			}
-
-			AsyncTask(ENamedThreads::GetRenderThread(), [Width, Height, this]() {
-				FScopeLock Lock(&RenderSyncContext);
-
-				const FIntPoint			  FrameSize = FIntPoint(Width, Height);
-				FRHICommandListImmediate& RHICmdList = FRHICommandListExecutor::GetImmediateCommandList();
-
-				if (!RenderTargetDescriptor.IsValid() || RenderTargetDescriptor.GetSize() != FIntVector(FrameSize.X, FrameSize.Y, 0))
-				{
-					// Create the RenderTarget descriptor
-					RenderTargetDescriptor = FPooledRenderTargetDesc::Create2DDesc(FrameSize,
-						PF_B8G8R8A8,
-						FClearValueBinding::None,
-						TexCreate_None,
-						TexCreate_RenderTargetable,
-						false);
-
-					// Update the shader resource for the 'SourceTexture'
-					FRHITextureCreateDesc RenderTargetTextureDesc =
-						FRHITextureCreateDesc::Create2D(TEXT(""), FrameSize.X, FrameSize.Y, PF_B8G8R8A8)
-							.SetClearValue(FClearValueBinding::None)
-#if PLATFORM_MAC
-							.SetFlags(ETextureCreateFlags::CPUReadback | ETextureCreateFlags::SRGB)
-							.SetInitialState(ERHIAccess::CPURead);
-#else
-							.SetFlags(ETextureCreateFlags::Dynamic | ETextureCreateFlags::ShaderResource | TexCreate_RenderTargetable | ETextureCreateFlags::SRGB)
-							.SetInitialState(ERHIAccess::SRVMask);
-#endif
-
-					SourceTexture = RHICreateTexture(RenderTargetTextureDesc);
-
-					// Find a free target-able texture from the render pool
-					GRenderTargetPool.FindFreeElement(RHICmdList,
-						RenderTargetDescriptor,
-						RenderTarget,
-						TEXT("PIXELSTEAMINGPLAYER"));
-				}
-
-				// Create the update region structure
-				const FUpdateTextureRegion2D Region(0, 0, 0, 0, FrameSize.X, FrameSize.Y);
-
-				// Set the Pixel data of the webrtc Frame to the SourceTexture
-				RHIUpdateTexture2D(SourceTexture, 0, Region, FrameSize.X * 4, Buffer.GetData());
-
-				CallConsumeFrame(SourceTexture);
-			});
-		}
-	}
-
 	void FVideoSink::CallConsumeFrame(FTextureRHIRef Frame)
 	{
 		// Iterate video consumers and pass this data to their buffers
diff --git a/Engine/Plugins/Media/PixelStreaming2/Source/PixelStreaming2/Private/VideoSourceGroup.cpp b/Engine/Plugins/Media/PixelStreaming2/Source/PixelStreaming2/Private/VideoSourceGroup.cpp
--- a/Engine/Plugins/Media/PixelStreaming2/Source/PixelStreaming2/Private/VideoSourceGroup.cpp
+++ b/Engine/Plugins/Media/PixelStreaming2/Source/PixelStreaming2/Private/VideoSourceGroup.cpp
@@ -2,11 +2,10 @@
 
 #include "VideoSourceGroup.h"
 
-#include "EpicRtcVideoSource.h"
 #include "Logging.h"
 #include "PixelStreaming2PluginSettings.h"
 #include "PixelStreaming2Trace.h"
-#include "Stats.h"
+#include "VideoSource.h"
 
 namespace UE::PixelStreaming2
 {
@@ -216,7 +215,7 @@ namespace UE::PixelStreaming2
 		}
 	}
 
-	void FVideoSourceGroup::AddVideoSource(TSharedPtr<FEpicRtcVideoSource> VideoSource)
+	void FVideoSourceGroup::AddVideoSource(TSharedPtr<FVideoSource> VideoSource)
 	{
 		{
 			FScopeLock Lock(&CriticalSection);
@@ -224,11 +223,11 @@ namespace UE::PixelStreaming2
 		}
 	}
 
-	void FVideoSourceGroup::RemoveVideoSource(const FEpicRtcVideoSource* ToRemove)
+	void FVideoSourceGroup::RemoveVideoSource(const FVideoSource* ToRemove)
 	{
 		{
 			FScopeLock Lock(&CriticalSection);
-			VideoSources.RemoveAll([ToRemove](const TSharedPtr<FEpicRtcVideoSource>& Target) {
+			VideoSources.RemoveAll([ToRemove](const TSharedPtr<FVideoSource>& Target) {
 				return Target.Get() == ToRemove;
 			});
 		}
@@ -321,7 +320,7 @@ namespace UE::PixelStreaming2
 
 	void FVideoSourceGroup::ForceKeyFrame()
 	{
-		TArray<TSharedPtr<FEpicRtcVideoSource>> VideoSourcesCopy;
+		TArray<TSharedPtr<FVideoSource>> VideoSourcesCopy;
 		{
 			// Grab a copy of the SharedPtr inside the lock to make sure it does not change elsewhere
 			// while allowing the functions to be called on the video source outside the lock.
diff --git a/Engine/Plugins/Media/PixelStreaming2/Source/PixelStreaming2/Public/IPixelStreaming2AudioConsumer.h b/Engine/Plugins/Media/PixelStreaming2/Source/PixelStreaming2/Public/IPixelStreaming2AudioConsumer.h
--- a/Engine/Plugins/Media/PixelStreaming2/Source/PixelStreaming2/Public/IPixelStreaming2AudioConsumer.h
+++ b/Engine/Plugins/Media/PixelStreaming2/Source/PixelStreaming2/Public/IPixelStreaming2AudioConsumer.h
@@ -2,7 +2,16 @@
 
 #pragma once
 
-#include "CoreMinimal.h"
+#include "CoreTypes.h"
+#include "UObject/Interface.h"
+
+#include "IPixelStreaming2AudioConsumer.generated.h"
+
+UINTERFACE(MinimalAPI)
+class UPixelStreaming2AudioConsumer : public UInterface
+{
+	GENERATED_BODY()
+};
 
 /**
  * An "Audio Consumer" is an object that is responsible for outputting the audio received from a peer. For example, by
@@ -10,9 +19,9 @@
  */
 class PIXELSTREAMING2_API IPixelStreaming2AudioConsumer
 {
-public:
-	virtual ~IPixelStreaming2AudioConsumer() = default;
+	GENERATED_BODY()
 
+public:
 	/**
 	 * @brief Consume raw audio data.
 	 * @param AudioData Pointer to the audio data.
diff --git a/Engine/Plugins/Media/PixelStreaming2/Source/PixelStreaming2/Public/IPixelStreaming2AudioProducer.h b/Engine/Plugins/Media/PixelStreaming2/Source/PixelStreaming2/Public/IPixelStreaming2AudioProducer.h
--- a/Engine/Plugins/Media/PixelStreaming2/Source/PixelStreaming2/Public/IPixelStreaming2AudioProducer.h
+++ b/Engine/Plugins/Media/PixelStreaming2/Source/PixelStreaming2/Public/IPixelStreaming2AudioProducer.h
@@ -3,6 +3,15 @@
 #pragma once
 
 #include "CoreTypes.h"
+#include "UObject/Interface.h"
+
+#include "IPixelStreaming2AudioProducer.generated.h"
+
+UINTERFACE(MinimalAPI)
+class UPixelStreaming2AudioProducer : public UInterface
+{
+	GENERATED_BODY()
+};
 
 /**
  * A "Audio Producer" is an object that you use to push audio frames into the Pixel Streaming system.
@@ -15,10 +24,9 @@
  */
 class PIXELSTREAMING2_API IPixelStreaming2AudioProducer
 {
-public:
-	IPixelStreaming2AudioProducer() = default;
-	virtual ~IPixelStreaming2AudioProducer() = default;
+	GENERATED_BODY()
 
+public:
 	/**
 	 * @brief Pushes audio into the PS2 audio pipeline which will mix with other audio producers before broadcasting.
 	 * @param InBuffer Pointer to the audio data.
diff --git a/Engine/Plugins/Media/PixelStreaming2/Source/PixelStreaming2/Public/IPixelStreaming2AudioSink.h b/Engine/Plugins/Media/PixelStreaming2/Source/PixelStreaming2/Public/IPixelStreaming2AudioSink.h
--- a/Engine/Plugins/Media/PixelStreaming2/Source/PixelStreaming2/Public/IPixelStreaming2AudioSink.h
+++ b/Engine/Plugins/Media/PixelStreaming2/Source/PixelStreaming2/Public/IPixelStreaming2AudioSink.h
@@ -2,9 +2,17 @@
 
 #pragma once
 
-#include "CoreMinimal.h"
+#include "CoreTypes.h"
+#include "IPixelStreaming2AudioConsumer.h"
+#include "UObject/Interface.h"
 
-class IPixelStreaming2AudioConsumer;
+#include "IPixelStreaming2AudioSink.generated.h"
+
+UINTERFACE(MinimalAPI)
+class UPixelStreaming2AudioSink : public UInterface
+{
+	GENERATED_BODY()
+};
 
 /**
  * An "Audio Sink" is an object that receives audio from a singular peer. From here, you can add multiple consumers to
@@ -12,9 +20,9 @@ class IPixelStreaming2AudioConsumer;
  */
 class PIXELSTREAMING2_API IPixelStreaming2AudioSink
 {
-public:
-	virtual ~IPixelStreaming2AudioSink() = default;
+	GENERATED_BODY()
 
+public:
 	/**
 	 * @brief Add an audio consumer to the sink.
 	 * @param AudioConsumer The Audio consumer to add to the sink.
diff --git a/Engine/Plugins/Media/PixelStreaming2/Source/PixelStreaming2/Public/IPixelStreaming2Module.h b/Engine/Plugins/Media/PixelStreaming2/Source/PixelStreaming2/Public/IPixelStreaming2Module.h
--- a/Engine/Plugins/Media/PixelStreaming2/Source/PixelStreaming2/Public/IPixelStreaming2Module.h
+++ b/Engine/Plugins/Media/PixelStreaming2/Source/PixelStreaming2/Public/IPixelStreaming2Module.h
@@ -2,15 +2,14 @@
 
 #pragma once
 
-#include "CoreMinimal.h"
+#include "CoreTypes.h"
 #include "Modules/ModuleInterface.h"
 #include "Modules/ModuleManager.h"
-#include "Templates/SharedPointer.h"
 #include "Templates/RefCounting.h"
+#include "Templates/SharedPointer.h"
 #include "IPixelStreaming2AudioProducer.h"
 #include "IPixelStreaming2Streamer.h"
-
-#include "Video/VideoConfig.h"
+#include "IPixelStreaming2VideoProducer.h"
 
 /**
  * The IPixelStreaming2Module interface manages the core functionality of the Pixel Streaming system.
@@ -60,7 +59,7 @@ public:
 
 	/**
 	 * Starts streaming on all streamers.
-	 * @return False if the module was not able to start streaming. 
+	 * @return False if the module was not able to start streaming.
 	 */
 	virtual bool StartStreaming() = 0;
 
@@ -74,7 +73,7 @@ public:
 	 * @param StreamerId - The ID of the Streamer to be created.
 	 * @return SharedPtr to the streamer.
 	 */
-	virtual TSharedPtr<IPixelStreaming2Streamer> CreateStreamer(const FString& StreamerId) = 0;
+	virtual TSharedPtr<IPixelStreaming2Streamer> CreateStreamer(const FString& StreamerId, const FString& Type = TEXT("DefaultRtc")) = 0;
 
 	/**
 	 * @brief Creates a new video producer.
@@ -121,11 +120,18 @@ public:
 	 */
 	virtual FString GetDefaultStreamerID() = 0;
 
+	/**
+	 * Get the Default Connection URL ("ws://127.0.0.1:8888").
+	 * @return FString The default connection url ("ws://127.0.0.1:8888").
+	 */
+	virtual FString GetDefaultConnectionURL() = 0;
+
 	/**
 	 * Get the Default Signaling URL ("ws://127.0.0.1:8888").
 	 * @return FString The default signaling url ("ws://127.0.0.1:8888").
 	 */
-	virtual FString GetDefaultSignallingURL() = 0;
+	UE_DEPRECATED(5.6, "GetDefaultSignallingURL has been deprecated. Please use GetDefaultConnectionURL instead.")
+	virtual FString GetDefaultSignallingURL() { return GetDefaultConnectionURL(); }
 
 	/**
 	 * @brief A method for iterating through all of the streamers on the module.
@@ -133,4 +139,4 @@ public:
 	 * @param Func The lambda to execute with each streamer.
 	 */
 	virtual void ForEachStreamer(const TFunction<void(TSharedPtr<IPixelStreaming2Streamer>)>& Func) = 0;
-};
+};
\ No newline at end of file
diff --git a/Engine/Plugins/Media/PixelStreaming2/Source/PixelStreaming2/Public/IPixelStreaming2VideoConsumer.h b/Engine/Plugins/Media/PixelStreaming2/Source/PixelStreaming2/Public/IPixelStreaming2VideoConsumer.h
--- a/Engine/Plugins/Media/PixelStreaming2/Source/PixelStreaming2/Public/IPixelStreaming2VideoConsumer.h
+++ b/Engine/Plugins/Media/PixelStreaming2/Source/PixelStreaming2/Public/IPixelStreaming2VideoConsumer.h
@@ -2,8 +2,17 @@
 
 #pragma once
 
-#include "CoreMinimal.h"
-#include "Rendering/SlateRenderer.h"
+#include "CoreTypes.h"
+#include "UObject/Interface.h"
+#include "RHIFwd.h"
+
+#include "IPixelStreaming2VideoConsumer.generated.h"
+
+UINTERFACE(MinimalAPI)
+class UPixelStreaming2VideoConsumer : public UInterface
+{
+	GENERATED_BODY()
+};
 
 /**
  * An "Video Consumer" is an object that is responsible for outputting the video received from a peer. For example, by
@@ -11,9 +20,9 @@
  */
 class PIXELSTREAMING2_API IPixelStreaming2VideoConsumer
 {
-public:
-	virtual ~IPixelStreaming2VideoConsumer() = default;
+	GENERATED_BODY()
 
+public:
 	/**
 	 * @brief Consume a texture as a video frame.
 	 * @param Frame The Frame to consume.
diff --git a/Engine/Plugins/Media/PixelStreaming2/Source/PixelStreaming2/Public/IPixelStreaming2VideoProducer.h b/Engine/Plugins/Media/PixelStreaming2/Source/PixelStreaming2/Public/IPixelStreaming2VideoProducer.h
--- a/Engine/Plugins/Media/PixelStreaming2/Source/PixelStreaming2/Public/IPixelStreaming2VideoProducer.h
+++ b/Engine/Plugins/Media/PixelStreaming2/Source/PixelStreaming2/Public/IPixelStreaming2VideoProducer.h
@@ -4,6 +4,15 @@
 
 #include "CoreTypes.h"
 #include "IPixelCaptureInputFrame.h"
+#include "UObject/Interface.h"
+
+#include "IPixelStreaming2VideoProducer.generated.h"
+
+UINTERFACE(MinimalAPI)
+class UPixelStreaming2VideoProducer : public UInterface
+{
+	GENERATED_BODY()
+};
 
 /**
  * A "Video Producer" is an object that you use to push video frames into the Pixel Streaming system.
@@ -17,10 +26,9 @@
  */
 class PIXELSTREAMING2_API IPixelStreaming2VideoProducer
 {
-public:
-    IPixelStreaming2VideoProducer() = default;
-	virtual ~IPixelStreaming2VideoProducer() = default;
+	GENERATED_BODY()
 
+public:
 	/**
 	 * Pushes a raw video frame into the Pixel Streaming system.
 	 * @param InputFrame The raw input frame, which may be a GPU texture or a CPU texture based on the underlying type of frame pushed.
diff --git a/Engine/Plugins/Media/PixelStreaming2/Source/PixelStreaming2/Public/IPixelStreaming2VideoSink.h b/Engine/Plugins/Media/PixelStreaming2/Source/PixelStreaming2/Public/IPixelStreaming2VideoSink.h
--- a/Engine/Plugins/Media/PixelStreaming2/Source/PixelStreaming2/Public/IPixelStreaming2VideoSink.h
+++ b/Engine/Plugins/Media/PixelStreaming2/Source/PixelStreaming2/Public/IPixelStreaming2VideoSink.h
@@ -2,7 +2,17 @@
 
 #pragma once
 
-#include "CoreMinimal.h"
+#include "CoreTypes.h"
+#include "IPixelStreaming2VideoConsumer.h"
+#include "UObject/Interface.h"
+
+#include "IPixelStreaming2VideoSink.generated.h"
+
+UINTERFACE(MinimalAPI)
+class UPixelStreaming2VideoSink : public UInterface
+{
+	GENERATED_BODY()
+};
 
 class IPixelStreaming2VideoConsumer;
 
@@ -11,9 +21,9 @@ class IPixelStreaming2VideoConsumer;
  */
 class PIXELSTREAMING2_API IPixelStreaming2VideoSink
 {
-public:
-	virtual ~IPixelStreaming2VideoSink() = default;
+	GENERATED_BODY()
 
+public:
 	/**
 	 * @brief Add a video consumer to the sink.
 	 * @param VideoConsumer The video consumer to add to the sink.
diff --git a/Engine/Plugins/Media/PixelStreaming2/Source/PixelStreaming2/Public/PixelStreaming2Delegates.h b/Engine/Plugins/Media/PixelStreaming2/Source/PixelStreaming2/Public/PixelStreaming2Delegates.h
--- a/Engine/Plugins/Media/PixelStreaming2/Source/PixelStreaming2/Public/PixelStreaming2Delegates.h
+++ b/Engine/Plugins/Media/PixelStreaming2/Source/PixelStreaming2/Public/PixelStreaming2Delegates.h
@@ -2,8 +2,10 @@
 
 #pragma once
 
-#include "Delegates/Delegate.h"
 #include "Containers/UnrealString.h"
+#include "Delegates/Delegate.h"
+#include "UObject/Object.h"
+
 #include "PixelStreaming2Delegates.generated.h"
 
 /**
diff --git a/Engine/Plugins/Media/PixelStreaming2/Source/PixelStreaming2Editor/PixelStreaming2Editor.build.cs b/Engine/Plugins/Media/PixelStreaming2/Source/PixelStreaming2Editor/PixelStreaming2Editor.build.cs
--- a/Engine/Plugins/Media/PixelStreaming2/Source/PixelStreaming2Editor/PixelStreaming2Editor.build.cs
+++ b/Engine/Plugins/Media/PixelStreaming2/Source/PixelStreaming2Editor/PixelStreaming2Editor.build.cs
@@ -19,8 +19,8 @@ namespace UnrealBuildTool.Rules
 				"RenderCore",
 				"Renderer",
 				"RHI",
-				"PixelStreaming2",
-				"PixelStreaming2Core",
+				"PixelStreaming2RTC",
+				"PixelStreaming2Settings",
 				"Slate",
 				"SlateCore",
 				"EngineSettings",
@@ -32,7 +32,8 @@ namespace UnrealBuildTool.Rules
 				"Sockets",
 				"ApplicationCore",
 				"PixelStreaming2Input",
-				"AVCodecsCore"
+				"AVCodecsCore",
+				"PixelStreaming2"
 			});
 
 			if (Target.bBuildEditor)
diff --git a/Engine/Plugins/Media/PixelStreaming2/Source/PixelStreaming2Editor/Private/AssetTypeActions_VideoProducer.h b/Engine/Plugins/Media/PixelStreaming2/Source/PixelStreaming2Editor/Private/AssetTypeActions_VideoProducer.h
--- a/Engine/Plugins/Media/PixelStreaming2/Source/PixelStreaming2Editor/Private/AssetTypeActions_VideoProducer.h
+++ b/Engine/Plugins/Media/PixelStreaming2/Source/PixelStreaming2Editor/Private/AssetTypeActions_VideoProducer.h
@@ -13,7 +13,7 @@ public:
 	// IAssetTypeActions Implementation
 	virtual FText GetName() const override { return NSLOCTEXT("PixelStreaming2", "AssetTypeActions_VideoProducer", "Video Input Actions"); }
 	virtual FColor GetTypeColor() const override { return FColor(192,64,64); }
-	virtual UClass* GetSupportedClass() const override { return UPixelStreaming2VideoProducer::StaticClass(); }
+	virtual UClass* GetSupportedClass() const override { return UPixelStreaming2VideoProducerBase::StaticClass(); }
 	virtual bool CanFilter() override { return false; }
 	virtual uint32 GetCategories() override
 	{
