[PS2] PR #12718: Refactor PixelStreaming2 plugin to enable other streaming protocols

diff --git a/Engine/Plugins/Media/PixelStreaming2/Source/PixelStreaming2Editor/Private/PixelStreaming2EditorModule.cpp b/Engine/Plugins/Media/PixelStreaming2/Source/PixelStreaming2Editor/Private/PixelStreaming2EditorModule.cpp
--- a/Engine/Plugins/Media/PixelStreaming2/Source/PixelStreaming2Editor/Private/PixelStreaming2EditorModule.cpp
+++ b/Engine/Plugins/Media/PixelStreaming2/Source/PixelStreaming2Editor/Private/PixelStreaming2EditorModule.cpp
@@ -6,10 +6,10 @@
 #include "PixelStreaming2Toolbar.h"
 #include "AssetTypeActions_VideoProducer.h"
 #include "Engine/GameViewportClient.h"
-#include "CoreUtils.h"
+#include "UtilsCore.h"
 #include "Editor/EditorPerformanceSettings.h"
 #include "Framework/Application/SlateApplication.h"
-#include "IPixelStreaming2Module.h"
+#include "IPixelStreaming2RTCModule.h"
 #include "IPixelStreaming2Streamer.h"
 #include "VideoProducerBackBufferComposited.h"
 #include "VideoProducerViewport.h"
@@ -29,6 +29,7 @@
 #include "Framework/Notifications/NotificationManager.h"
 #include "Widgets/Notifications/SNotificationList.h"
 #include "Logging.h"
+#include "IPixelStreaming2Module.h"
 
 #define LOCTEXT_NAMESPACE "PixelStreaming2EditorModule"
 
@@ -106,7 +107,7 @@ namespace UE::EditorPixelStreaming2
 			}),
 			ECVF_Default);
 
-		IPixelStreaming2Module& Module = IPixelStreaming2Module::Get();
+		IPixelStreaming2RTCModule& Module = IPixelStreaming2RTCModule::Get();
 		Module.OnReady().AddRaw(this, &FPixelStreaming2EditorModule::InitEditorStreaming);
 	}
 
@@ -120,7 +121,7 @@ namespace UE::EditorPixelStreaming2
 		StopStreaming();
 	}
 
-	void FPixelStreaming2EditorModule::InitEditorStreaming(IPixelStreaming2Module& Module)
+	void FPixelStreaming2EditorModule::InitEditorStreaming(IPixelStreaming2RTCModule& Module)
 	{
 		FString EditorStreamerID;
 		if (!FParse::Value(FCommandLine::Get(), TEXT("PixelStreamingID="), EditorStreamerID))
@@ -128,7 +129,7 @@ namespace UE::EditorPixelStreaming2
 			EditorStreamerID = "Editor";
 		}
 
-		EditorStreamer = Module.CreateStreamer(EditorStreamerID);
+		EditorStreamer = IPixelStreaming2Module::Get().CreateStreamer(EditorStreamerID);
 
 		// Bind to start/stop streaming so we disable/restore relevant editor settings
 		EditorStreamer->OnStreamingStarted().AddLambda([this](IPixelStreaming2Streamer* Streamer) {
@@ -149,17 +150,17 @@ namespace UE::EditorPixelStreaming2
 		});
 
 		// Give the editor streamer the default url if the user hasn't specified one when launching the editor
-		if (EditorStreamer->GetSignallingServerURL().IsEmpty())
+		if (EditorStreamer->GetConnectionURL().IsEmpty())
 		{
 			// No URL was passed on the command line, initialize defaults
 			StreamerPort = 8888;
 			SignallingDomain = TEXT("ws://127.0.0.1");
 
-			EditorStreamer->SetSignallingServerURL(FString::Printf(TEXT("%s:%d"), *SignallingDomain, StreamerPort));
+			EditorStreamer->SetConnectionURL(FString::Printf(TEXT("%s:%d"), *SignallingDomain, StreamerPort));
 		}
 		else
 		{
-			FString			  SpecifiedSignallingURL = EditorStreamer->GetSignallingServerURL();
+			FString			  SpecifiedSignallingURL = EditorStreamer->GetConnectionURL();
 			TOptional<uint16> ExtractedStreamerPort = FGenericPlatformHttp::GetUrlPort(SpecifiedSignallingURL);
 			StreamerPort = (int32)ExtractedStreamerPort.Get(8888);
 
@@ -264,7 +265,7 @@ namespace UE::EditorPixelStreaming2
 
 		if (!UPixelStreaming2PluginSettings::CVarEditorUseRemoteSignallingServer.GetValueOnAnyThread())
 		{
-			EditorStreamer->SetSignallingServerURL(FString::Printf(TEXT("%s:%d"), *SignallingDomain, StreamerPort));
+			EditorStreamer->SetConnectionURL(FString::Printf(TEXT("%s:%d"), *SignallingDomain, StreamerPort));
 			StartSignalling();
 		}
 
diff --git a/Engine/Plugins/Media/PixelStreaming2/Source/PixelStreaming2Editor/Private/PixelStreaming2EditorModule.h b/Engine/Plugins/Media/PixelStreaming2/Source/PixelStreaming2Editor/Private/PixelStreaming2EditorModule.h
--- a/Engine/Plugins/Media/PixelStreaming2/Source/PixelStreaming2Editor/Private/PixelStreaming2EditorModule.h
+++ b/Engine/Plugins/Media/PixelStreaming2/Source/PixelStreaming2Editor/Private/PixelStreaming2EditorModule.h
@@ -3,7 +3,7 @@
 #pragma once
 
 #include "IPixelStreaming2EditorModule.h"
-#include "IPixelStreaming2Module.h"
+#include "IPixelStreaming2RTCModule.h"
 #include "PixelStreaming2Servers.h"
 
 namespace UE::EditorPixelStreaming2
@@ -32,7 +32,7 @@ namespace UE::EditorPixelStreaming2
 		virtual int32	GetViewerPort() override { return ViewerPort; };
 
 	private:
-		void InitEditorStreaming(IPixelStreaming2Module& Module);
+		void InitEditorStreaming(IPixelStreaming2RTCModule& Module);
 		bool ParseResolution(const TCHAR* InResolution, uint32& OutX, uint32& OutY);
 		void MaybeResizeEditor(TSharedPtr<SWindow> RootWindow);
 		void DisableCPUThrottlingSetting();
diff --git a/Engine/Plugins/Media/PixelStreaming2/Source/PixelStreaming2Editor/Private/PixelStreaming2Toolbar.cpp b/Engine/Plugins/Media/PixelStreaming2/Source/PixelStreaming2Editor/Private/PixelStreaming2Toolbar.cpp
--- a/Engine/Plugins/Media/PixelStreaming2/Source/PixelStreaming2Editor/Private/PixelStreaming2Toolbar.cpp
+++ b/Engine/Plugins/Media/PixelStreaming2/Source/PixelStreaming2Editor/Private/PixelStreaming2Toolbar.cpp
@@ -34,7 +34,7 @@
 #include "Widgets/Input/SNumericEntryBox.h"
 #include "PixelStreaming2PluginSettings.h"
 
-#include "CoderUtils.h"
+#include "UtilsCoder.h"
 #include "Video/Encoders/Configs/VideoEncoderConfigAV1.h"
 #include "Video/Encoders/Configs/VideoEncoderConfigH264.h"
 
@@ -359,16 +359,16 @@ namespace UE::EditorPixelStreaming2
 						  [SNew(SEditableTextBox)
 								  .Text_Lambda([]() {
 									  TSharedPtr<IPixelStreaming2Streamer> Streamer = IPixelStreaming2Module::Get().FindStreamer("Editor");
-									  return FText::FromString(Streamer->GetSignallingServerURL());
+									  return FText::FromString(Streamer->GetConnectionURL());
 								  })
 								  .OnTextChanged_Lambda([](const FText& InText) {
 									  IPixelStreaming2Module::Get().ForEachStreamer([InText](TSharedPtr<IPixelStreaming2Streamer> Streamer) {
-										  Streamer->SetSignallingServerURL(InText.ToString());
+										  Streamer->SetConnectionURL(InText.ToString());
 									  });
 								  })
 								  .OnTextCommitted_Lambda([](const FText& InText, ETextCommit::Type InTextCommit) {
 									  IPixelStreaming2Module::Get().ForEachStreamer([InText](TSharedPtr<IPixelStreaming2Streamer> Streamer) {
-										  Streamer->SetSignallingServerURL(InText.ToString());
+										  Streamer->SetConnectionURL(InText.ToString());
 									  });
 								  })
 								  .IsEnabled_Lambda([]() {
diff --git a/Engine/Plugins/Media/PixelStreaming2/Source/PixelStreaming2Editor/Private/PixelStreaming2Toolbar.h b/Engine/Plugins/Media/PixelStreaming2/Source/PixelStreaming2Editor/Private/PixelStreaming2Toolbar.h
--- a/Engine/Plugins/Media/PixelStreaming2/Source/PixelStreaming2Editor/Private/PixelStreaming2Toolbar.h
+++ b/Engine/Plugins/Media/PixelStreaming2/Source/PixelStreaming2Editor/Private/PixelStreaming2Toolbar.h
@@ -8,7 +8,6 @@
 #include "Framework/MultiBox/MultiBoxBuilder.h"
 #include "LevelEditor.h"
 #include "ToolMenus.h"
-#include "IPixelStreaming2Module.h"
 
 namespace UE::EditorPixelStreaming2
 {
diff --git a/Engine/Plugins/Media/PixelStreaming2/Source/PixelStreaming2Editor/Private/VideoProducerViewport.cpp b/Engine/Plugins/Media/PixelStreaming2/Source/PixelStreaming2Editor/Private/VideoProducerViewport.cpp
--- a/Engine/Plugins/Media/PixelStreaming2/Source/PixelStreaming2Editor/Private/VideoProducerViewport.cpp
+++ b/Engine/Plugins/Media/PixelStreaming2/Source/PixelStreaming2Editor/Private/VideoProducerViewport.cpp
@@ -5,9 +5,9 @@
 #include "Async/Async.h"
 #include "Engine/GameViewportClient.h"
 #include "Framework/Application/SlateApplication.h"
-#include "IPixelStreaming2Module.h"
 #include "LevelEditor.h"
 #include "LevelEditorViewport.h"
+#include "Modules/ModuleManager.h"
 #include "PixelCaptureBufferFormat.h"
 #include "PixelCaptureCapturerRHI.h"
 #include "PixelCaptureInputFrameRHI.h"
@@ -16,6 +16,7 @@
 #include "UnrealClient.h"
 #include "Utils.h"
 
+
 namespace UE::EditorPixelStreaming2
 {
 
diff --git a/Engine/Plugins/Media/PixelStreaming2/Source/PixelStreaming2Editor/Public/IPixelStreaming2EditorModule.h b/Engine/Plugins/Media/PixelStreaming2/Source/PixelStreaming2Editor/Public/IPixelStreaming2EditorModule.h
--- a/Engine/Plugins/Media/PixelStreaming2/Source/PixelStreaming2Editor/Public/IPixelStreaming2EditorModule.h
+++ b/Engine/Plugins/Media/PixelStreaming2/Source/PixelStreaming2Editor/Public/IPixelStreaming2EditorModule.h
@@ -5,7 +5,7 @@
 #include "CoreMinimal.h"
 #include "Modules/ModuleInterface.h"
 #include "Modules/ModuleManager.h"
-#include "PixelStreaming2CoreEnums.h"
+#include "PixelStreaming2SettingsEnums.h"
 #include "PixelStreaming2Servers.h"
 
 /**
diff --git a/Engine/Plugins/Media/PixelStreaming2/Source/PixelStreaming2HMD/PixelStreaming2HMD.build.cs b/Engine/Plugins/Media/PixelStreaming2/Source/PixelStreaming2HMD/PixelStreaming2HMD.build.cs
--- a/Engine/Plugins/Media/PixelStreaming2/Source/PixelStreaming2HMD/PixelStreaming2HMD.build.cs
+++ b/Engine/Plugins/Media/PixelStreaming2/Source/PixelStreaming2HMD/PixelStreaming2HMD.build.cs
@@ -19,7 +19,7 @@ namespace UnrealBuildTool.Rules
 					"HeadMountedDisplay",
 					"XRBase",
 					"SlateCore",
-					"PixelStreaming2Core"
+					"PixelStreaming2Settings"
 				}
 			);
 		}
diff --git a/Engine/Plugins/Media/PixelStreaming2/Source/PixelStreaming2HMD/Private/PixelStreaming2HMDModule.cpp b/Engine/Plugins/Media/PixelStreaming2/Source/PixelStreaming2HMD/Private/PixelStreaming2HMDModule.cpp
--- a/Engine/Plugins/Media/PixelStreaming2/Source/PixelStreaming2HMD/Private/PixelStreaming2HMDModule.cpp
+++ b/Engine/Plugins/Media/PixelStreaming2/Source/PixelStreaming2HMD/Private/PixelStreaming2HMDModule.cpp
@@ -2,7 +2,6 @@
 
 #include "PixelStreaming2HMDModule.h"
 
-#include "CoreUtils.h"
 #include "IXRTrackingSystem.h"
 #include "PixelStreaming2HMD.h"
 #include "PixelStreaming2PluginSettings.h"
@@ -14,22 +13,12 @@ namespace UE::PixelStreaming2HMD
 	 */
 	void FPixelStreaming2HMDModule::StartupModule()
 	{
-		if (!UE::PixelStreaming2::IsStreamingSupported())
-		{
-			return;
-		}
-
 		IModularFeatures::Get().RegisterModularFeature(GetModularFeatureName(), this);
 		ActiveXRSystem = EPixelStreaming2XRSystem::Unknown;
 	}
 
 	void FPixelStreaming2HMDModule::ShutdownModule()
 	{
-		if (!UE::PixelStreaming2::IsStreamingSupported())
-		{
-			return;
-		}
-
 		// Remove the modules hold of the ptr
 		// HMD.Reset();
 		IModularFeatures::Get().UnregisterModularFeature(GetModularFeatureName(), this);
diff --git a/Engine/Plugins/Media/PixelStreaming2/Source/PixelStreaming2Input/PixelStreaming2Input.build.cs b/Engine/Plugins/Media/PixelStreaming2/Source/PixelStreaming2Input/PixelStreaming2Input.build.cs
--- a/Engine/Plugins/Media/PixelStreaming2/Source/PixelStreaming2Input/PixelStreaming2Input.build.cs
+++ b/Engine/Plugins/Media/PixelStreaming2/Source/PixelStreaming2Input/PixelStreaming2Input.build.cs
@@ -12,14 +12,13 @@ namespace UnrealBuildTool.Rules
 					"Core",
 					"CoreUObject",
 					"Engine",
-					"InputDevice",
 					"Json",
 					"SlateCore",
 					"Slate",
 					"DeveloperSettings",
 					"HeadMountedDisplay",
 					"XRBase",
-					"PixelStreaming2Core"
+					"PixelStreaming2Settings"
 				}
 			);
 
@@ -27,9 +26,10 @@ namespace UnrealBuildTool.Rules
 				new string[]
 				{
 					"ApplicationCore",
-					"PixelStreaming2HMD",
+					"Core",
 					"InputCore",
-					"Core"
+					"InputDevice",
+					"PixelStreaming2HMD"
 				}
 			);
 		}
diff --git a/Engine/Plugins/Media/PixelStreaming2/Source/PixelStreaming2Input/Private/PixelStreaming2InputModule.cpp b/Engine/Plugins/Media/PixelStreaming2/Source/PixelStreaming2Input/Private/PixelStreaming2InputModule.cpp
--- a/Engine/Plugins/Media/PixelStreaming2/Source/PixelStreaming2Input/Private/PixelStreaming2InputModule.cpp
+++ b/Engine/Plugins/Media/PixelStreaming2/Source/PixelStreaming2Input/Private/PixelStreaming2InputModule.cpp
@@ -3,7 +3,6 @@
 #include "PixelStreaming2InputModule.h"
 
 #include "ApplicationWrapper.h"
-#include "CoreUtils.h"
 #include "Framework/Application/SlateApplication.h"
 #include "InputHandler.h"
 #include "IPixelStreaming2HMDModule.h"
@@ -13,11 +12,6 @@ namespace UE::PixelStreaming2Input
 
 	void FPixelStreaming2InputModule::StartupModule()
 	{
-		if (!UE::PixelStreaming2::IsStreamingSupported())
-		{
-			return;
-		}
-
 		if (!FSlateApplication::IsInitialized())
 		{
 			return;
@@ -30,11 +24,6 @@ namespace UE::PixelStreaming2Input
 
 	void FPixelStreaming2InputModule::ShutdownModule()
 	{
-		if (!UE::PixelStreaming2::IsStreamingSupported())
-		{
-			return;
-		}
-
 		IModularFeatures::Get().UnregisterModularFeature(GetModularFeatureName(), this);
 	}
 
diff --git a/Engine/Plugins/Media/PixelStreaming2/Source/PixelStreaming2RTC/Internal/EpicRtcAudioTrackObserver.h b/Engine/Plugins/Media/PixelStreaming2/Source/PixelStreaming2RTC/Internal/EpicRtcAudioTrackObserver.h
new file mode 100644
--- /dev/null
+++ b/Engine/Plugins/Media/PixelStreaming2/Source/PixelStreaming2RTC/Internal/EpicRtcAudioTrackObserver.h
@@ -0,0 +1,36 @@
+// Copyright Epic Games, Inc. All Rights Reserved.
+
+#pragma once
+
+#include "Templates/SharedPointer.h"
+
+#include "epic_rtc/core/audio/audio_track_observer.h"
+#include "epic_rtc_helper/memory/ref_count_impl_helper.h"
+
+namespace UE::PixelStreaming2
+{
+	class FEpicRtcManager;
+
+	class PIXELSTREAMING2RTC_API FEpicRtcAudioTrackObserver : public EpicRtcAudioTrackObserverInterface
+	{
+	public:
+		FEpicRtcAudioTrackObserver(TWeakPtr<FEpicRtcManager> Manager);
+		virtual ~FEpicRtcAudioTrackObserver() = default;
+
+	private:
+		// Begin EpicRtcAudioTrackObserverInterface
+		virtual void OnAudioTrackMuted(EpicRtcAudioTrackInterface* AudioTrack, EpicRtcBool bIsMuted) override;
+		virtual void OnAudioTrackFrame(EpicRtcAudioTrackInterface* AudioTrack, const EpicRtcAudioFrame& Frame) override;
+		virtual void OnAudioTrackRemoved(EpicRtcAudioTrackInterface* AudioTrack) override;
+		virtual void OnAudioTrackState(EpicRtcAudioTrackInterface* AudioTrack, const EpicRtcTrackState State) override;
+		// End EpicRtcAudioTrackObserverInterface
+
+	public:
+		// Begin EpicRtcRefCountInterface
+		EPICRTC_REFCOUNT_INTERFACE_IN_PLACE
+		// End EpicRtcRefCountInterface
+
+	private:
+		TWeakPtr<FEpicRtcManager> Manager;
+	};
+} // namespace UE::PixelStreaming2
\ No newline at end of file
diff --git a/Engine/Plugins/Media/PixelStreaming2/Source/PixelStreaming2RTC/Internal/EpicRtcAudioTrackObserverFactory.h b/Engine/Plugins/Media/PixelStreaming2/Source/PixelStreaming2RTC/Internal/EpicRtcAudioTrackObserverFactory.h
new file mode 100644
--- /dev/null
+++ b/Engine/Plugins/Media/PixelStreaming2/Source/PixelStreaming2RTC/Internal/EpicRtcAudioTrackObserverFactory.h
@@ -0,0 +1,35 @@
+// Copyright Epic Games, Inc. All Rights Reserved.
+
+#pragma once
+
+#include "EpicRtcAudioTrackObserver.h"
+#include "Templates/SharedPointer.h"
+
+#include "epic_rtc/core/audio/audio_track_observer.h"
+#include "epic_rtc_helper/memory/ref_count_impl_helper.h"
+
+namespace UE::PixelStreaming2
+{
+	class FEpicRtcManager;
+
+	class PIXELSTREAMING2RTC_API FEpicRtcAudioTrackObserverFactory : public EpicRtcAudioTrackObserverFactoryInterface
+	{
+	public:
+		FEpicRtcAudioTrackObserverFactory(TWeakPtr<FEpicRtcManager> Manager);
+		virtual ~FEpicRtcAudioTrackObserverFactory() = default;
+
+	public:
+		// Begin EpicRtcAudioTrackObserverFactoryInterface 
+		virtual EpicRtcErrorCode CreateAudioTrackObserver(const EpicRtcStringView ParticipantId, const EpicRtcStringView AudioTrackId, EpicRtcAudioTrackObserverInterface** OutAudioTrackObserver) override;
+		// End EpicRtcAudioTrackObserverFactoryInterface 
+
+	public:
+		// Begin EpicRtcRefCountInterface 
+		EPICRTC_REFCOUNT_INTERFACE_IN_PLACE
+		// End EpicRtcRefCountInterface 
+
+	private:
+		TWeakPtr<FEpicRtcManager> Manager;
+	};
+
+} // namespace UE::PixelStreaming2
\ No newline at end of file
diff --git a/Engine/Plugins/Media/PixelStreaming2/Source/PixelStreaming2RTC/Internal/EpicRtcConferenceUtils.h b/Engine/Plugins/Media/PixelStreaming2/Source/PixelStreaming2RTC/Internal/EpicRtcConferenceUtils.h
new file mode 100644
--- /dev/null
+++ b/Engine/Plugins/Media/PixelStreaming2/Source/PixelStreaming2RTC/Internal/EpicRtcConferenceUtils.h
@@ -0,0 +1,86 @@
+// Copyright Epic Games, Inc. All Rights Reserved.
+
+#pragma once
+
+#include "IPixelStreaming2Stats.h"
+#include "PixelStreaming2PluginSettings.h"
+#include "TickableTask.h"
+
+#include "epic_rtc/core/conference.h"
+
+namespace UE::PixelStreaming2
+{
+	class PIXELSTREAMING2RTC_API FEpicRtcTickConferenceTask : public FPixelStreamingTickableTask
+	{
+	public:
+		FEpicRtcTickConferenceTask(TRefCountPtr<EpicRtcConferenceInterface> EpicRtcConference, const FString& TaskName = TEXT("EpicRtcTickConferenceTask"))
+			: EpicRtcConference(EpicRtcConference)
+			, TaskName(TaskName)
+		{
+		}
+
+		virtual ~FEpicRtcTickConferenceTask() override
+		{
+			// We may get a call to destroy the task before we've had a chance to tick again.
+			// So to be safe, we tick the conference a final time
+			if (EpicRtcConference)
+			{
+
+				while (EpicRtcConference->NeedsTick())
+				{
+					EpicRtcConference->Tick();
+				}
+			}
+		}
+
+		// Begin FPixelStreamingTickableTask
+		virtual void Tick(float DeltaMs) override
+		{
+			if (EpicRtcConference)
+			{
+				MsSinceLastAudioTick += DeltaMs;
+				MsSinceLastStatsTick += DeltaMs;
+
+				// Tick conference normally. This handles things like data channel message
+				IPixelStreaming2Stats::Get().GraphValue(TEXT("ConferenceTickInterval"), DeltaMs, 1, 0.f, 1.f);
+				while (EpicRtcConference->NeedsTick())
+				{
+					EpicRtcConference->Tick();
+				}
+
+				// Tick audio (every 10 ms). This enables the pulling of audio from the ADM
+				if (MsSinceLastAudioTick >= 10.f)
+				{
+					// Track the interval. This is helpful for seeing if we're going over the 10ms requirement
+					IPixelStreaming2Stats::Get().GraphValue(TEXT("AudioTickInterval"), MsSinceLastAudioTick, 1, 0.f, 25.f, 10.f);
+					EpicRtcConference->TickAudio();
+					MsSinceLastAudioTick = 0.f;
+				}
+
+				// Tick stats at the configured interval
+				float StatsInterval = UPixelStreaming2PluginSettings::CVarWebRTCStatsInterval.GetValueOnAnyThread() * 1000.f;
+				bool  bStatsEnabled = !UPixelStreaming2PluginSettings::CVarWebRTCDisableStats.GetValueOnAnyThread();
+				if (MsSinceLastStatsTick >= StatsInterval && StatsInterval > 0.f && bStatsEnabled)
+				{
+					IPixelStreaming2Stats::Get().GraphValue(TEXT("StatTickInterval"), MsSinceLastStatsTick, 1, 0.f, 25.f, 10.f);
+					EpicRtcConference->TickStats();
+					MsSinceLastStatsTick = 0.f;
+				}
+			}
+		}
+
+		virtual const FString& GetName() const override
+		{
+			return TaskName;
+		}
+		// End FPixelStreamingTickableTask
+
+	private:
+		TRefCountPtr<EpicRtcConferenceInterface> EpicRtcConference;
+		FString									 TaskName;
+
+		float MsSinceLastAudioTick = 0.f;
+		float MsSinceLastStatsTick = 0.f;
+	};
+
+} // namespace UE::PixelStreaming2
\ No newline at end of file
diff --git a/Engine/Plugins/Media/PixelStreaming2/Source/PixelStreaming2RTC/Internal/EpicRtcDataTrackObserver.h b/Engine/Plugins/Media/PixelStreaming2/Source/PixelStreaming2RTC/Internal/EpicRtcDataTrackObserver.h
new file mode 100644
--- /dev/null
+++ b/Engine/Plugins/Media/PixelStreaming2/Source/PixelStreaming2RTC/Internal/EpicRtcDataTrackObserver.h
@@ -0,0 +1,36 @@
+// Copyright Epic Games, Inc. All Rights Reserved.
+
+#pragma once
+
+#include "Templates/SharedPointer.h"
+
+#include "epic_rtc/core/data_track_observer.h"
+#include "epic_rtc_helper/memory/ref_count_impl_helper.h"
+
+namespace UE::PixelStreaming2
+{
+	class FEpicRtcManager;
+
+	class PIXELSTREAMING2RTC_API FEpicRtcDataTrackObserver : public EpicRtcDataTrackObserverInterface
+	{
+	public:
+		FEpicRtcDataTrackObserver(TWeakPtr<FEpicRtcManager> Manager);
+		virtual ~FEpicRtcDataTrackObserver() = default;
+
+	private:
+		// Begin EpicRtcDataTrackObserverInterface
+		virtual void OnDataTrackState(EpicRtcDataTrackInterface* DataTrack, const EpicRtcTrackState State) override;
+		virtual void OnDataTrackMessage(EpicRtcDataTrackInterface* DataTrack) override;
+		virtual void OnDataTrackError(EpicRtcDataTrackInterface*, const EpicRtcErrorCode) override {};
+		// End EpicRtcDataTrackObserverInterface
+
+	public:
+		// Begin EpicRtcRefCountInterface
+		EPICRTC_REFCOUNT_INTERFACE_IN_PLACE
+		// End EpicRtcRefCountInterface
+
+	private:
+		TWeakPtr<FEpicRtcManager> Manager;
+	};
+
+} // namespace UE::PixelStreaming2
\ No newline at end of file
diff --git a/Engine/Plugins/Media/PixelStreaming2/Source/PixelStreaming2RTC/Internal/EpicRtcDataTrackObserverFactory.h b/Engine/Plugins/Media/PixelStreaming2/Source/PixelStreaming2RTC/Internal/EpicRtcDataTrackObserverFactory.h
new file mode 100644
--- /dev/null
+++ b/Engine/Plugins/Media/PixelStreaming2/Source/PixelStreaming2RTC/Internal/EpicRtcDataTrackObserverFactory.h
@@ -0,0 +1,35 @@
+// Copyright Epic Games, Inc. All Rights Reserved.
+
+#pragma once
+
+#include "EpicRtcDataTrackObserver.h"
+#include "Templates/SharedPointer.h"
+
+#include "epic_rtc/core/data_track_observer.h"
+#include "epic_rtc_helper/memory/ref_count_impl_helper.h"
+
+namespace UE::PixelStreaming2
+{
+	class FEpicRtcManager;
+
+	class PIXELSTREAMING2RTC_API FEpicRtcDataTrackObserverFactory : public EpicRtcDataTrackObserverFactoryInterface
+	{
+	public:
+		FEpicRtcDataTrackObserverFactory(TWeakPtr<FEpicRtcManager> Manager);
+		virtual ~FEpicRtcDataTrackObserverFactory() = default;
+
+	public:
+		// Begin EpicRtcDataTrackObserverFactoryInterface 
+		virtual EpicRtcErrorCode CreateDataTrackObserver(const EpicRtcStringView ParticipantId, const EpicRtcStringView DataTrackId, EpicRtcDataTrackObserverInterface** OutDataTrackObserver) override;
+		// End EpicRtcDataTrackObserverFactoryInterface 
+
+	public:
+		// Begin EpicRtcRefCountInterface 
+		EPICRTC_REFCOUNT_INTERFACE_IN_PLACE
+		// End EpicRtcRefCountInterface 
+
+	private:
+		TWeakPtr<FEpicRtcManager> Manager;
+	};
+
+} // namespace UE::PixelStreaming2
\ No newline at end of file
diff --git a/Engine/Plugins/Media/PixelStreaming2/Source/PixelStreaming2RTC/Internal/EpicRtcLogging.h b/Engine/Plugins/Media/PixelStreaming2/Source/PixelStreaming2RTC/Internal/EpicRtcLogging.h
new file mode 100644
--- /dev/null
+++ b/Engine/Plugins/Media/PixelStreaming2/Source/PixelStreaming2RTC/Internal/EpicRtcLogging.h
@@ -0,0 +1,19 @@
+// Copyright Epic Games, Inc. All Rights Reserved.
+
+#pragma once
+
+#include "Logging/LogMacros.h"
+
+#include "epic_rtc/common/logging.h"
+
+DECLARE_LOG_CATEGORY_EXTERN(LogPixelStreaming2EpicRtc, Log, All);
+DECLARE_LOG_CATEGORY_EXTERN(LogPixelStreaming2WebRtc, Fatal, All);
+
+namespace UE::PixelStreaming2
+{
+	class PIXELSTREAMING2RTC_API FEpicRtcLogsRedirector : public EpicRtcLoggerInterface
+	{
+	public:
+		virtual void Log(const EpicRtcLogMessage& Message) override;
+	};
+} // namespace UE::PixelStreaming2
\ No newline at end of file
diff --git a/Engine/Plugins/Media/PixelStreaming2/Source/PixelStreaming2RTC/Internal/EpicRtcManager.h b/Engine/Plugins/Media/PixelStreaming2/Source/PixelStreaming2RTC/Internal/EpicRtcManager.h
new file mode 100644
--- /dev/null
+++ b/Engine/Plugins/Media/PixelStreaming2/Source/PixelStreaming2RTC/Internal/EpicRtcManager.h
@@ -0,0 +1,125 @@
+// Copyright Epic Games, Inc. All Rights Reserved.
+
+#pragma once
+
+#include "Delegates/Delegate.h"
+#include "EpicRtcAudioTrackObserverFactory.h"
+#include "EpicRtcDataTrackObserverFactory.h"
+#include "EpicRtcRoomObserver.h"
+#include "EpicRtcSessionObserver.h"
+#include "EpicRtcVideoTrackObserverFactory.h"
+#include "Templates/RefCounting.h"
+
+#include "epic_rtc/core/conference.h"
+#include "epic_rtc/core/room.h"
+#include "epic_rtc/core/session.h"
+#include "epic_rtc/plugins/signalling/websocket.h"
+
+namespace UE::PixelStreaming2
+{
+	/**
+	 * A base class for managing epic rtc sessions, rooms and participants
+	 *
+	 * Extending this class allows observers to callback to the manager
+	 */
+	class PIXELSTREAMING2RTC_API FEpicRtcManager
+	{
+	public:
+		virtual ~FEpicRtcManager() = default;
+
+	public:
+		// Begin FEpicRtcSessionObserver Callbacks
+		DECLARE_TS_MULTICAST_DELEGATE_OneParam(FOnSessionStateUpdate, const EpicRtcSessionState);
+		FOnSessionStateUpdate OnSessionStateUpdate;
+
+		DECLARE_TS_MULTICAST_DELEGATE_OneParam(FOnSessionErrorUpdate, const EpicRtcErrorCode);
+		FOnSessionErrorUpdate OnSessionErrorUpdate;
+
+		DECLARE_TS_MULTICAST_DELEGATE_OneParam(FOnSessionRoomsAvailableUpdate, EpicRtcStringArrayInterface*);
+		FOnSessionRoomsAvailableUpdate OnSessionRoomsAvailableUpdate;
+		// End FEpicRtcSessionObserver Callbacks
+
+		// Begin FEpicRtcRoomObserver Callbacks
+		DECLARE_TS_MULTICAST_DELEGATE_OneParam(FOnRoomStateUpdate, const EpicRtcRoomState);
+		FOnRoomStateUpdate OnRoomStateUpdate;
+
+		DECLARE_TS_MULTICAST_DELEGATE_OneParam(FOnRoomJoinedUpdate, EpicRtcParticipantInterface*);
+		FOnRoomJoinedUpdate OnRoomJoinedUpdate;
+
+		DECLARE_TS_MULTICAST_DELEGATE_OneParam(FOnRoomLeftUpdate, const EpicRtcStringView);
+		FOnRoomLeftUpdate OnRoomLeftUpdate;
+
+		DECLARE_TS_MULTICAST_DELEGATE_TwoParams(FOnAudioTrackUpdate, EpicRtcParticipantInterface*, EpicRtcAudioTrackInterface*);
+		FOnAudioTrackUpdate OnAudioTrackUpdate;
+
+		DECLARE_TS_MULTICAST_DELEGATE_TwoParams(FOnVideoTrackUpdate, EpicRtcParticipantInterface*, EpicRtcVideoTrackInterface*);
+		FOnVideoTrackUpdate OnVideoTrackUpdate;
+
+		DECLARE_TS_MULTICAST_DELEGATE_TwoParams(FOnDataTrackUpdate, EpicRtcParticipantInterface*, EpicRtcDataTrackInterface*);
+		FOnDataTrackUpdate OnDataTrackUpdate;
+
+		DECLARE_TS_MULTICAST_DELEGATE_TwoParams(FOnLocalSdpUpdate, EpicRtcParticipantInterface*, EpicRtcSdpInterface*);
+		FOnLocalSdpUpdate OnLocalSdpUpdate;
+
+		DECLARE_TS_MULTICAST_DELEGATE_TwoParams(FOnRemoteSdpUpdate, EpicRtcParticipantInterface*, EpicRtcSdpInterface*);
+		FOnRemoteSdpUpdate OnRemoteSdpUpdate;
+
+		DECLARE_TS_MULTICAST_DELEGATE_OneParam(FOnRoomErrorUpdate, const EpicRtcErrorCode);
+		FOnRoomErrorUpdate OnRoomErrorUpdate;
+		// End FEpicRtcRoomObserver Callbacks
+
+		// Begin FEpicRtcAudioTrackObserver Callbacks
+		DECLARE_TS_MULTICAST_DELEGATE_TwoParams(FOnAudioTrackMuted, EpicRtcAudioTrackInterface*, EpicRtcBool);
+		FOnAudioTrackMuted OnAudioTrackMuted;
+
+		DECLARE_TS_MULTICAST_DELEGATE_TwoParams(FOnAudioTrackFrame, EpicRtcAudioTrackInterface*, const EpicRtcAudioFrame&);
+		FOnAudioTrackFrame OnAudioTrackFrame;
+
+		DECLARE_TS_MULTICAST_DELEGATE_OneParam(FOnAudioTrackRemoved, EpicRtcAudioTrackInterface*);
+		FOnAudioTrackRemoved OnAudioTrackRemoved;
+
+		DECLARE_TS_MULTICAST_DELEGATE_TwoParams(FOnAudioTrackState, EpicRtcAudioTrackInterface*, const EpicRtcTrackState);
+		FOnAudioTrackState OnAudioTrackState;
+		// End FEpicRtcAudioTrackObserver Callbacks
+
+		// Begin FEpicRtcVideoTrackObserver Callbacks
+		DECLARE_TS_MULTICAST_DELEGATE_TwoParams(FOnVideoTrackMuted, EpicRtcVideoTrackInterface*, EpicRtcBool);
+		FOnVideoTrackMuted OnVideoTrackMuted;
+
+		DECLARE_TS_MULTICAST_DELEGATE_TwoParams(FOnVideoTrackFrame, EpicRtcVideoTrackInterface*, const EpicRtcVideoFrame&);
+		FOnVideoTrackFrame OnVideoTrackFrame;
+
+		DECLARE_TS_MULTICAST_DELEGATE_OneParam(FOnVideoTrackRemoved, EpicRtcVideoTrackInterface*);
+		FOnVideoTrackRemoved OnVideoTrackRemoved;
+
+		DECLARE_TS_MULTICAST_DELEGATE_TwoParams(FOnVideoTrackState, EpicRtcVideoTrackInterface*, const EpicRtcTrackState);
+		FOnVideoTrackState OnVideoTrackState;
+		// End FEpicRtcVideoTrackObserver Callbacks
+
+		// Begin FEpicRtcDataTrackObserver Callbacks
+		DECLARE_TS_MULTICAST_DELEGATE_OneParam(FOnDataTrackRemoved, EpicRtcDataTrackInterface*);
+		FOnDataTrackRemoved OnDataTrackRemoved;
+
+		DECLARE_TS_MULTICAST_DELEGATE_TwoParams(FOnDataTrackState, EpicRtcDataTrackInterface*, const EpicRtcTrackState);
+		FOnDataTrackState OnDataTrackState;
+
+		DECLARE_TS_MULTICAST_DELEGATE_OneParam(FOnDataTrackMessage, EpicRtcDataTrackInterface*);
+		FOnDataTrackMessage OnDataTrackMessage;
+		// End FEpicRtcDataTrackObserver Callbacks
+
+	public:
+		// Begin EpicRtc Classes
+		TRefCountPtr<EpicRtcConferenceInterface> EpicRtcConference;
+		TRefCountPtr<EpicRtcSessionInterface>	 EpicRtcSession;
+		TRefCountPtr<EpicRtcRoomInterface>		 EpicRtcRoom;
+		// End EpicRtc Classes
+
+		// Begin EpicRtc Observers
+		TRefCountPtr<FEpicRtcSessionObserver>			SessionObserver;
+		TRefCountPtr<FEpicRtcRoomObserver>				RoomObserver;
+		TRefCountPtr<FEpicRtcAudioTrackObserverFactory> AudioTrackObserverFactory;
+		TRefCountPtr<FEpicRtcVideoTrackObserverFactory> VideoTrackObserverFactory;
+		TRefCountPtr<FEpicRtcDataTrackObserverFactory>	DataTrackObserverFactory;
+		// End EpicRtc Observers
+	};
+} // namespace UE::PixelStreaming2
\ No newline at end of file
diff --git a/Engine/Plugins/Media/PixelStreaming2/Source/PixelStreaming2RTC/Internal/EpicRtcRoomObserver.h b/Engine/Plugins/Media/PixelStreaming2/Source/PixelStreaming2RTC/Internal/EpicRtcRoomObserver.h
new file mode 100644
--- /dev/null
+++ b/Engine/Plugins/Media/PixelStreaming2/Source/PixelStreaming2RTC/Internal/EpicRtcRoomObserver.h
@@ -0,0 +1,42 @@
+// Copyright Epic Games, Inc. All Rights Reserved.
+
+#pragma once
+
+#include "Templates/SharedPointer.h"
+
+#include "epic_rtc/core/room_observer.h"
+#include "epic_rtc_helper/memory/ref_count_impl_helper.h"
+
+namespace UE::PixelStreaming2
+{
+	class FEpicRtcManager;
+
+	class PIXELSTREAMING2RTC_API FEpicRtcRoomObserver : public EpicRtcRoomObserverInterface
+	{
+	public:
+		FEpicRtcRoomObserver(TWeakPtr<FEpicRtcManager> Manager);
+		virtual ~FEpicRtcRoomObserver() = default;
+
+	private:
+		// Begin EpicRtcRoomObserver
+		virtual void							   OnRoomStateUpdate(const EpicRtcRoomState State) override;
+		virtual void							   OnRoomJoinedUpdate(EpicRtcParticipantInterface* Participant) override;
+		virtual void							   OnRoomLeftUpdate(const EpicRtcStringView ParticipantId) override;
+		virtual void							   OnAudioTrackUpdate(EpicRtcParticipantInterface* Participant, EpicRtcAudioTrackInterface* AudioTrack) override;
+		virtual void							   OnVideoTrackUpdate(EpicRtcParticipantInterface* Participant, EpicRtcVideoTrackInterface* VideoTrack) override;
+		virtual void							   OnDataTrackUpdate(EpicRtcParticipantInterface* Participant, EpicRtcDataTrackInterface* DataTrack) override;
+		[[nodiscard]] virtual EpicRtcSdpInterface* OnLocalSdpUpdate(EpicRtcParticipantInterface* Participant, EpicRtcSdpInterface* Sdp) override;
+		[[nodiscard]] virtual EpicRtcSdpInterface* OnRemoteSdpUpdate(EpicRtcParticipantInterface* Participant, EpicRtcSdpInterface* Sdp) override;
+		virtual void							   OnRoomErrorUpdate(const EpicRtcErrorCode Error) override;
+		// Begin EpicRtcRoomObserver
+
+	public:
+		// Begin EpicRtcRefCountInterface
+		EPICRTC_REFCOUNT_INTERFACE_IN_PLACE
+		// End EpicRtcRefCountInterface
+
+	private:
+		TWeakPtr<FEpicRtcManager> Manager;
+	};
+
+} // namespace UE::PixelStreaming2
\ No newline at end of file
diff --git a/Engine/Plugins/Media/PixelStreaming2/Source/PixelStreaming2RTC/Internal/EpicRtcSessionObserver.h b/Engine/Plugins/Media/PixelStreaming2/Source/PixelStreaming2RTC/Internal/EpicRtcSessionObserver.h
new file mode 100644
--- /dev/null
+++ b/Engine/Plugins/Media/PixelStreaming2/Source/PixelStreaming2RTC/Internal/EpicRtcSessionObserver.h
@@ -0,0 +1,34 @@
+// Copyright Epic Games, Inc. All Rights Reserved.
+
+#pragma once
+
+#include "Templates/SharedPointer.h"
+
+#include "epic_rtc/core/session_observer.h"
+#include "epic_rtc_helper/memory/ref_count_impl_helper.h"
+
+namespace UE::PixelStreaming2
+{
+	class FEpicRtcManager;
+
+	class PIXELSTREAMING2RTC_API FEpicRtcSessionObserver : public EpicRtcSessionObserverInterface
+	{
+	public:
+		FEpicRtcSessionObserver(TWeakPtr<FEpicRtcManager> Manager);
+		virtual ~FEpicRtcSessionObserver() = default;
+
+	private:
+		// Begin EpicRtcSessionObserver
+		virtual void OnSessionStateUpdate(const EpicRtcSessionState State) override;
+		virtual void OnSessionErrorUpdate(const EpicRtcErrorCode Error) override;
+		virtual void OnSessionRoomsAvailableUpdate(EpicRtcStringArrayInterface* RoomsList) override;
+		// End EpicRtcSessionObserver
+	public:
+		// Begin EpicRtcRefCountInterface
+		EPICRTC_REFCOUNT_INTERFACE_IN_PLACE
+		// End EpicRtcRefCountInterface
+	private:
+		TWeakPtr<FEpicRtcManager> Manager;
+	};
+
+} // namespace UE::PixelStreaming2
diff --git a/Engine/Plugins/Media/PixelStreaming2/Source/PixelStreaming2RTC/Internal/EpicRtcVideoCommon.h b/Engine/Plugins/Media/PixelStreaming2/Source/PixelStreaming2RTC/Internal/EpicRtcVideoCommon.h
new file mode 100644
--- /dev/null
+++ b/Engine/Plugins/Media/PixelStreaming2/Source/PixelStreaming2RTC/Internal/EpicRtcVideoCommon.h
@@ -0,0 +1,842 @@
+// Copyright Epic Games, Inc. All Rights Reserved.
+
+#pragma once
+
+#include "Video/DependencyDescriptor.h"
+#include "Video/GenericFrameInfo.h"
+#include "Video/VideoEncoder.h"
+
+#include "epic_rtc/containers/epic_rtc_array.h"
+#include "epic_rtc/containers/epic_rtc_string_view.h"
+#include "epic_rtc/core/video/video_buffer.h"
+#include "epic_rtc/core/video/video_codec_info.h"
+#include "epic_rtc_helper/memory/ref_count_impl_helper.h"
+
+FORCEINLINE bool operator==(const EpicRtcVideoResolution& Lhs, const EpicRtcVideoResolution& Rhs)
+{
+	return Lhs._width == Rhs._width && Lhs._height == Rhs._height;
+}
+
+namespace UE::PixelStreaming2
+{
+	class PIXELSTREAMING2RTC_API FEpicRtcEncodedVideoBuffer : public EpicRtcEncodedVideoBufferInterface
+	{
+	public:
+		FEpicRtcEncodedVideoBuffer() = default;
+		FEpicRtcEncodedVideoBuffer(uint8_t* InData, uint64_t InSize)
+			: Data(InData, InSize)
+		{
+		}
+		virtual ~FEpicRtcEncodedVideoBuffer() = default;
+
+		virtual const uint8_t* GetData() const override { return Data.GetData(); };
+		virtual uint64_t	   GetSize() const override { return Data.Num(); };
+
+	private:
+		TArray<uint8> Data;
+
+	public:
+		// Begin EpicRtcRefCountInterface
+		EPICRTC_REFCOUNT_INTERFACE_IN_PLACE
+		// End EpicRtcRefCountInterface
+	};
+
+	class PIXELSTREAMING2RTC_API FEpicRtcParameterPairArray : public EpicRtcParameterPairArrayInterface
+	{
+	public:
+		FEpicRtcParameterPairArray() = default;
+		virtual ~FEpicRtcParameterPairArray() = default;
+
+		FEpicRtcParameterPairArray(const TArray<EpicRtcParameterPair>& ParameterPairs)
+			: Data(ParameterPairs)
+		{
+		}
+
+		FEpicRtcParameterPairArray(std::initializer_list<EpicRtcParameterPair> ParameterPairs)
+			: Data(ParameterPairs)
+		{
+		}
+
+		virtual const EpicRtcParameterPair* Get() const override
+		{
+			return Data.GetData();
+		}
+
+		virtual EpicRtcParameterPair* Get() override
+		{
+			return Data.GetData();
+		}
+
+		virtual uint64_t Size() const override
+		{
+			return Data.Num();
+		}
+
+		void Append(std::initializer_list<EpicRtcParameterPair> ParameterPairs)
+		{
+			Data.Append(ParameterPairs);
+		}
+
+	private:
+		TArray<EpicRtcParameterPair> Data;
+
+	public:
+		// Begin EpicRtcRefCountInterface
+		EPICRTC_REFCOUNT_INTERFACE_IN_PLACE
+		// End EpicRtcRefCountInterface
+	};
+
+	class PIXELSTREAMING2RTC_API FEpicRtcScalabilityModeArray : public EpicRtcVideoScalabilityModeArrayInterface
+	{
+	public:
+		FEpicRtcScalabilityModeArray() = default;
+		virtual ~FEpicRtcScalabilityModeArray() = default;
+
+		FEpicRtcScalabilityModeArray(const TArray<EpicRtcVideoScalabilityMode>& ScalabilityModes)
+			: Data(ScalabilityModes)
+		{
+		}
+
+		FEpicRtcScalabilityModeArray(std::initializer_list<EpicRtcVideoScalabilityMode> ScalabilityModes)
+			: Data(ScalabilityModes)
+		{
+		}
+
+		FEpicRtcScalabilityModeArray(const TArray<EScalabilityMode>& ScalabilityModes)
+		{
+			for (EScalabilityMode ScalabilityMode : ScalabilityModes)
+			{
+				Data.Add(static_cast<EpicRtcVideoScalabilityMode>(ScalabilityMode)); // HACK if the Enums become un-aligned
+			}
+		}
+
+		virtual const EpicRtcVideoScalabilityMode* Get() const override
+		{
+			return Data.GetData();
+		}
+
+		virtual EpicRtcVideoScalabilityMode* Get() override
+		{
+			return Data.GetData();
+		}
+
+		virtual uint64_t Size() const override
+		{
+			return Data.Num();
+		}
+
+		void Append(std::initializer_list<EpicRtcVideoScalabilityMode> ScalabilityModes)
+		{
+			Data.Append(ScalabilityModes);
+		}
+
+	private:
+		TArray<EpicRtcVideoScalabilityMode> Data;
+
+	public:
+		// Begin EpicRtcRefCountInterface
+		EPICRTC_REFCOUNT_INTERFACE_IN_PLACE
+		// End EpicRtcRefCountInterface
+	};
+
+	class PIXELSTREAMING2RTC_API FEpicRtcVideoCodecInfo : public EpicRtcVideoCodecInfoInterface
+	{
+	public:
+		FEpicRtcVideoCodecInfo(EpicRtcVideoCodec Codec, bool bIsHardwareAccelerated = false, EpicRtcParameterPairArrayInterface* Parameters = new FEpicRtcParameterPairArray(), EpicRtcVideoScalabilityModeArrayInterface* ScalabilityModes = new FEpicRtcScalabilityModeArray())
+			: Codec(Codec)
+			, bIsHardwareAccelerated(bIsHardwareAccelerated)
+			, Parameters(Parameters)
+			, ScalabilityModes(ScalabilityModes)
+		{
+		}
+
+		virtual ~FEpicRtcVideoCodecInfo() = default;
+
+		EpicRtcVideoCodec GetCodec() override
+		{
+			return Codec;
+		}
+
+		EpicRtcParameterPairArrayInterface* GetParameters()
+		{
+			return Parameters;
+		}
+
+		EpicRtcVideoScalabilityModeArrayInterface* GetScalabilityModes() override
+		{
+			return ScalabilityModes;
+		}
+
+		EpicRtcBool IsHardwareAccelerated() override
+		{
+			return bIsHardwareAccelerated;
+		}
+
+	private:
+		EpicRtcVideoCodec										Codec;
+		bool													bIsHardwareAccelerated;
+		TRefCountPtr<EpicRtcParameterPairArrayInterface>		Parameters;
+		TRefCountPtr<EpicRtcVideoScalabilityModeArrayInterface> ScalabilityModes;
+
+	public:
+		// Begin EpicRtcRefCountInterface
+		EPICRTC_REFCOUNT_INTERFACE_IN_PLACE
+		// End EpicRtcRefCountInterface
+	};
+
+	class PIXELSTREAMING2RTC_API FVideoCodecInfoArray : public EpicRtcVideoCodecInfoArrayInterface
+	{
+	public:
+		FVideoCodecInfoArray() = default;
+
+		virtual ~FVideoCodecInfoArray()
+		{
+			for (auto& Codec : Data)
+			{
+				if (Codec != nullptr)
+				{
+					Codec->Release();
+				}
+			}
+		}
+
+		FVideoCodecInfoArray(const TArray<TRefCountPtr<EpicRtcVideoCodecInfoInterface>>& Codecs)
+		{
+			for (auto& Codec : Codecs)
+			{
+				Data.Add(Codec.GetReference());
+				if (Codec.GetReference() != nullptr)
+				{
+					Codec->AddRef();
+				}
+			}
+		}
+
+		FVideoCodecInfoArray(const TArray<EpicRtcVideoCodecInfoInterface*>& Codecs)
+		{
+			for (auto& Codec : Codecs)
+			{
+				Data.Add(Codec);
+				if (Codec != nullptr)
+				{
+					Codec->AddRef();
+				}
+			}
+		}
+
+		FVideoCodecInfoArray(std::initializer_list<EpicRtcVideoCodecInfoInterface*> Codecs)
+		{
+			for (auto& Codec : Codecs)
+			{
+				Data.Add(Codec);
+				if (Codec != nullptr)
+				{
+					Codec->AddRef();
+				}
+			}
+		}
+
+		virtual EpicRtcVideoCodecInfoInterface* const* Get() const override
+		{
+			return Data.GetData();
+		}
+
+		virtual EpicRtcVideoCodecInfoInterface** Get() override
+		{
+			return Data.GetData();
+		}
+
+		virtual uint64_t Size() const override
+		{
+			return Data.Num();
+		}
+
+	private:
+		TArray<EpicRtcVideoCodecInfoInterface*> Data;
+
+	public:
+		// Begin EpicRtcRefCountInterface
+		EPICRTC_REFCOUNT_INTERFACE_IN_PLACE
+		// End EpicRtcRefCountInterface
+	};
+
+	class PIXELSTREAMING2RTC_API FEpicRtcVideoResolutionBitrateLimitsArray : public EpicRtcVideoResolutionBitrateLimitsArrayInterface
+	{
+	public:
+		FEpicRtcVideoResolutionBitrateLimitsArray() = default;
+		virtual ~FEpicRtcVideoResolutionBitrateLimitsArray() = default;
+
+		FEpicRtcVideoResolutionBitrateLimitsArray(const TArray<EpicRtcVideoResolutionBitrateLimits>& BitrateLimits)
+			: Data(BitrateLimits)
+		{
+		}
+
+		FEpicRtcVideoResolutionBitrateLimitsArray(std::initializer_list<EpicRtcVideoResolutionBitrateLimits> BitrateLimits)
+			: Data(BitrateLimits)
+		{
+		}
+
+		virtual const EpicRtcVideoResolutionBitrateLimits* Get() const override
+		{
+			return Data.GetData();
+		}
+
+		virtual EpicRtcVideoResolutionBitrateLimits* Get() override
+		{
+			return Data.GetData();
+		}
+
+		virtual uint64_t Size() const override
+		{
+			return Data.Num();
+		}
+
+	private:
+		TArray<EpicRtcVideoResolutionBitrateLimits> Data;
+
+	public:
+		// Begin EpicRtcRefCountInterface
+		EPICRTC_REFCOUNT_INTERFACE_IN_PLACE
+		// End EpicRtcRefCountInterface
+	};
+
+	class PIXELSTREAMING2RTC_API FEpicRtcPixelFormatArray : public EpicRtcPixelFormatArrayInterface
+	{
+	public:
+		FEpicRtcPixelFormatArray() = default;
+		virtual ~FEpicRtcPixelFormatArray() = default;
+
+		FEpicRtcPixelFormatArray(const TArray<EpicRtcPixelFormat>& PixelFormats)
+			: Data(PixelFormats)
+		{
+		}
+
+		FEpicRtcPixelFormatArray(std::initializer_list<EpicRtcPixelFormat> PixelFormats)
+			: Data(PixelFormats)
+		{
+		}
+
+		virtual const EpicRtcPixelFormat* Get() const override
+		{
+			return Data.GetData();
+		}
+
+		virtual EpicRtcPixelFormat* Get() override
+		{
+			return Data.GetData();
+		}
+
+		virtual uint64_t Size() const override
+		{
+			return Data.Num();
+		}
+
+	private:
+		TArray<EpicRtcPixelFormat> Data;
+
+	public:
+		// Begin EpicRtcRefCountInterface
+		EPICRTC_REFCOUNT_INTERFACE_IN_PLACE
+		// End EpicRtcRefCountInterface
+	};
+
+	class PIXELSTREAMING2RTC_API FEpicRtcVideoFrameTypeArray : public EpicRtcVideoFrameTypeArrayInterface
+	{
+	public:
+		FEpicRtcVideoFrameTypeArray() = default;
+		virtual ~FEpicRtcVideoFrameTypeArray() = default;
+
+		FEpicRtcVideoFrameTypeArray(const TArray<EpicRtcVideoFrameType>& FrameTypes)
+			: Data(FrameTypes)
+		{
+		}
+
+		FEpicRtcVideoFrameTypeArray(std::initializer_list<EpicRtcVideoFrameType> FrameTypes)
+			: Data(FrameTypes)
+		{
+		}
+
+		virtual const EpicRtcVideoFrameType* Get() const override
+		{
+			return Data.GetData();
+		}
+
+		virtual uint64_t Size() const override
+		{
+			return Data.Num();
+		}
+
+	private:
+		TArray<EpicRtcVideoFrameType> Data;
+
+	public:
+		// Begin EpicRtcRefCountInterface
+		EPICRTC_REFCOUNT_INTERFACE_IN_PLACE
+		// End EpicRtcRefCountInterface
+	};
+
+	class PIXELSTREAMING2RTC_API FEpicRtcInt32Array : public EpicRtcInt32ArrayInterface
+	{
+	public:
+		FEpicRtcInt32Array() = default;
+		virtual ~FEpicRtcInt32Array() = default;
+
+		FEpicRtcInt32Array(const TArray<int32_t>& Ints)
+			: Data(Ints)
+		{
+		}
+
+		FEpicRtcInt32Array(std::initializer_list<int32_t> Ints)
+			: Data(Ints)
+		{
+		}
+
+		virtual const int32_t* Get() const override
+		{
+			return Data.GetData();
+		}
+
+		virtual int32_t* Get() override
+		{
+			return Data.GetData();
+		}
+
+		virtual uint64_t Size() const override
+		{
+			return Data.Num();
+		}
+
+		void Append(std::initializer_list<int32_t> Ints)
+		{
+			Data.Append(Ints);
+		}
+
+	private:
+		TArray<int32_t> Data;
+
+	public:
+		// Begin EpicRtcRefCountInterface
+		EPICRTC_REFCOUNT_INTERFACE_IN_PLACE
+		// End EpicRtcRefCountInterface
+	};
+
+	class PIXELSTREAMING2RTC_API FEpicRtcBoolArray : public EpicRtcBoolArrayInterface
+	{
+	public:
+		FEpicRtcBoolArray() = default;
+		virtual ~FEpicRtcBoolArray() = default;
+
+		FEpicRtcBoolArray(const TArray<EpicRtcBool>& Bools)
+			: Data(Bools)
+		{
+		}
+
+		FEpicRtcBoolArray(const TArray<bool>& Bools)
+		{
+			Data.SetNum(Bools.Num());
+			for (size_t i = 0; i < Bools.Num(); i++)
+			{
+				Data[i] = Bools[i];
+			}
+		}
+
+		FEpicRtcBoolArray(std::initializer_list<EpicRtcBool> Bools)
+			: Data(Bools)
+		{
+		}
+
+		FEpicRtcBoolArray(std::initializer_list<bool> Bools)
+		{
+			for (auto& Bool : Bools)
+			{
+				Data.Add(Bool);
+			}
+		}
+
+		virtual const EpicRtcBool* Get() const override
+		{
+			return Data.GetData();
+		}
+
+		virtual EpicRtcBool* Get() override
+		{
+			return Data.GetData();
+		}
+
+		virtual uint64_t Size() const override
+		{
+			return Data.Num();
+		}
+
+	private:
+		TArray<EpicRtcBool> Data;
+
+	public:
+		// Begin EpicRtcRefCountInterface
+		EPICRTC_REFCOUNT_INTERFACE_IN_PLACE
+		// End EpicRtcRefCountInterface
+	};
+
+	class PIXELSTREAMING2RTC_API FEpicRtcDecodeTargetIndicationArray : public EpicRtcDecodeTargetIndicationArrayInterface
+	{
+	public:
+		FEpicRtcDecodeTargetIndicationArray() = default;
+		virtual ~FEpicRtcDecodeTargetIndicationArray() = default;
+
+		FEpicRtcDecodeTargetIndicationArray(const TArray<EpicRtcDecodeTargetIndication>& DecodeTargetIndications)
+			: Data(DecodeTargetIndications)
+		{
+		}
+
+		FEpicRtcDecodeTargetIndicationArray(std::initializer_list<EpicRtcDecodeTargetIndication> DecodeTargetIndications)
+			: Data(DecodeTargetIndications)
+		{
+		}
+
+		// Helper method for converting array AVCodecs' EDecodeTargetIndication to array of EpicRtc's EpicRtcDecodeTargetIndication
+		FEpicRtcDecodeTargetIndicationArray(const TArray<EDecodeTargetIndication>& DecodeTargetIndications)
+		{
+			Data.SetNum(DecodeTargetIndications.Num());
+
+			for (size_t i = 0; i < DecodeTargetIndications.Num(); i++)
+			{
+				switch (DecodeTargetIndications[i])
+				{
+					case EDecodeTargetIndication::NotPresent:
+						Data[i] = EpicRtcDecodeTargetIndication::NotPresent;
+						break;
+					case EDecodeTargetIndication::Discardable:
+						Data[i] = EpicRtcDecodeTargetIndication::Discardable;
+						break;
+					case EDecodeTargetIndication::Switch:
+						Data[i] = EpicRtcDecodeTargetIndication::Switch;
+						break;
+					case EDecodeTargetIndication::Required:
+						Data[i] = EpicRtcDecodeTargetIndication::Required;
+						break;
+					default:
+						checkNoEntry();
+				}
+			}
+		}
+
+		virtual const EpicRtcDecodeTargetIndication* Get() const override
+		{
+			return Data.GetData();
+		}
+
+		virtual EpicRtcDecodeTargetIndication* Get() override
+		{
+			return Data.GetData();
+		}
+
+		virtual uint64_t Size() const override
+		{
+			return Data.Num();
+		}
+
+	private:
+		TArray<EpicRtcDecodeTargetIndication> Data;
+
+	public:
+		// Begin EpicRtcRefCountInterface
+		EPICRTC_REFCOUNT_INTERFACE_IN_PLACE
+		// End EpicRtcRefCountInterface
+	};
+
+	class PIXELSTREAMING2RTC_API FEpicRtcCodecBufferUsageArray : public EpicRtcCodecBufferUsageArrayInterface
+	{
+	public:
+		FEpicRtcCodecBufferUsageArray() = default;
+		virtual ~FEpicRtcCodecBufferUsageArray() = default;
+
+		FEpicRtcCodecBufferUsageArray(const TArray<EpicRtcCodecBufferUsage>& CodecBufferUsages)
+			: Data(CodecBufferUsages)
+		{
+		}
+
+		FEpicRtcCodecBufferUsageArray(std::initializer_list<EpicRtcCodecBufferUsage> CodecBufferUsages)
+			: Data(CodecBufferUsages)
+		{
+		}
+
+		// Helper method for converting array AVCodecs' FCodecBufferUsage to array of EpicRtc's EpicRtcCodecBufferUsage
+		FEpicRtcCodecBufferUsageArray(const TArray<FCodecBufferUsage>& CodecBufferUsages)
+		{
+			Data.SetNum(CodecBufferUsages.Num());
+			for (size_t i = 0; i < CodecBufferUsages.Num(); i++)
+			{
+				const FCodecBufferUsage& CodecBufferUsage = CodecBufferUsages[i];
+				Data[i] = EpicRtcCodecBufferUsage{
+					._id = CodecBufferUsage.Id,
+					._referenced = CodecBufferUsage.bReferenced,
+					._updated = CodecBufferUsage.bUpdated
+				};
+			}
+		}
+
+		virtual const EpicRtcCodecBufferUsage* Get() const override
+		{
+			return Data.GetData();
+		}
+
+		virtual EpicRtcCodecBufferUsage* Get() override
+		{
+			return Data.GetData();
+		}
+
+		virtual uint64_t Size() const override
+		{
+			return Data.Num();
+		}
+
+	private:
+		TArray<EpicRtcCodecBufferUsage> Data;
+
+	public:
+		// Begin EpicRtcRefCountInterface
+		EPICRTC_REFCOUNT_INTERFACE_IN_PLACE
+		// End EpicRtcRefCountInterface
+	};
+
+	class PIXELSTREAMING2RTC_API FEpicRtcVideoResolutionArray : public EpicRtcVideoResolutionArrayInterface
+	{
+	public:
+		FEpicRtcVideoResolutionArray() = default;
+		virtual ~FEpicRtcVideoResolutionArray() = default;
+
+		FEpicRtcVideoResolutionArray(const TArray<EpicRtcVideoResolution>& Resolutions)
+			: Data(Resolutions)
+		{
+		}
+
+		FEpicRtcVideoResolutionArray(std::initializer_list<EpicRtcVideoResolution> Resolutions)
+			: Data(Resolutions)
+		{
+		}
+
+		// Helper method for converting array AVCodecs' FResolution to array of EpicRtc's EpicRtcVideoResolution
+		FEpicRtcVideoResolutionArray(const TArray<FIntPoint>& Resolutions)
+		{
+			Data.SetNum(Resolutions.Num());
+			for (size_t i = 0; i < Resolutions.Num(); i++)
+			{
+				const FIntPoint& Resolution = Resolutions[i];
+				Data[i] = EpicRtcVideoResolution{
+					._width = Resolution.X,
+					._height = Resolution.Y
+				};
+			}
+		}
+
+		virtual const EpicRtcVideoResolution* Get() const override
+		{
+			return Data.GetData();
+		}
+
+		virtual EpicRtcVideoResolution* Get() override
+		{
+			return Data.GetData();
+		}
+
+		virtual uint64_t Size() const override
+		{
+			return Data.Num();
+		}
+
+	private:
+		TArray<EpicRtcVideoResolution> Data;
+
+	public:
+		// Begin EpicRtcRefCountInterface
+		EPICRTC_REFCOUNT_INTERFACE_IN_PLACE
+		// End EpicRtcRefCountInterface
+	};
+
+	class PIXELSTREAMING2RTC_API FEpicRtcGenericFrameInfoArray : public EpicRtcGenericFrameInfoArrayInterface
+	{
+	public:
+		FEpicRtcGenericFrameInfoArray() = default;
+
+		virtual ~FEpicRtcGenericFrameInfoArray()
+		{
+			for (auto& GenericFrameInfo : Data)
+			{
+				if (GenericFrameInfo != nullptr)
+				{
+					GenericFrameInfo->Release();
+				}
+			}
+		}
+
+		FEpicRtcGenericFrameInfoArray(const TArray<TRefCountPtr<EpicRtcGenericFrameInfoInterface>>& GenericFrameInfos)
+		{
+			for (auto& GenericFrameInfo : GenericFrameInfos)
+			{
+				Data.Add(GenericFrameInfo.GetReference());
+				if (GenericFrameInfo.GetReference() != nullptr)
+				{
+					GenericFrameInfo->AddRef();
+				}
+			}
+		}
+
+		FEpicRtcGenericFrameInfoArray(const TArray<EpicRtcGenericFrameInfoInterface*>& GenericFrameInfos)
+		{
+			for (auto& GenericFrameInfo : GenericFrameInfos)
+			{
+				Data.Add(GenericFrameInfo);
+				if (GenericFrameInfo != nullptr)
+				{
+					GenericFrameInfo->AddRef();
+				}
+			}
+		}
+
+		FEpicRtcGenericFrameInfoArray(std::initializer_list<EpicRtcGenericFrameInfoInterface*> GenericFrameInfos)
+		{
+			for (auto& GenericFrameInfo : GenericFrameInfos)
+			{
+				Data.Add(GenericFrameInfo);
+				if (GenericFrameInfo != nullptr)
+				{
+					GenericFrameInfo->AddRef();
+				}
+			}
+		}
+
+		virtual EpicRtcGenericFrameInfoInterface* const* Get() const override
+		{
+			return Data.GetData();
+		}
+
+		virtual EpicRtcGenericFrameInfoInterface** Get() override
+		{
+			return Data.GetData();
+		}
+
+		virtual uint64_t Size() const override
+		{
+			return Data.Num();
+		}
+
+	private:
+		TArray<EpicRtcGenericFrameInfoInterface*> Data;
+
+	public:
+		// Begin EpicRtcRefCountInterface
+		EPICRTC_REFCOUNT_INTERFACE_IN_PLACE
+		// End EpicRtcRefCountInterface
+	};
+
+	class PIXELSTREAMING2RTC_API FEpicRtcGenericFrameInfo : public EpicRtcGenericFrameInfoInterface
+	{
+	public:
+		FEpicRtcGenericFrameInfo(const FGenericFrameInfo& GenericFrameInfo)
+			: SpatialId(GenericFrameInfo.SpatialId)
+			, TemporalId(GenericFrameInfo.TemporalId)
+			, DecodeTargetIndications(MakeRefCount<FEpicRtcDecodeTargetIndicationArray>(GenericFrameInfo.DecodeTargetIndications))
+			, FrameDiffs(MakeRefCount<FEpicRtcInt32Array>(GenericFrameInfo.FrameDiffs))
+			, ChainDiffs(MakeRefCount<FEpicRtcInt32Array>(GenericFrameInfo.ChainDiffs))
+			, EncoderBuffers(MakeRefCount<FEpicRtcCodecBufferUsageArray>(GenericFrameInfo.EncoderBuffers))
+			, PartOfChain(MakeRefCount<FEpicRtcBoolArray>(GenericFrameInfo.PartOfChain))
+			, ActiveDecodeTargets(MakeRefCount<FEpicRtcBoolArray>(GenericFrameInfo.ActiveDecodeTargets))
+		{
+		}
+
+		virtual ~FEpicRtcGenericFrameInfo() = default;
+
+		virtual int32_t										 GetSpatialLayerId() override { return SpatialId; }
+		virtual int32_t										 GetTemporalLayerId() override { return TemporalId; }
+		virtual EpicRtcDecodeTargetIndicationArrayInterface* GetDecodeTargetIndications() override { return DecodeTargetIndications; }
+		virtual EpicRtcInt32ArrayInterface*					 GetFrameDiffs() override { return FrameDiffs; }
+		virtual EpicRtcInt32ArrayInterface*					 GetChainDiffs() override { return ChainDiffs; }
+		virtual EpicRtcCodecBufferUsageArrayInterface*		 GetEncoderBufferUsages() override { return EncoderBuffers; }
+		virtual EpicRtcBoolArrayInterface*					 GetPartOfChain() override { return PartOfChain; }
+		virtual EpicRtcBoolArrayInterface*					 GetActiveDecodeTargets() override { return ActiveDecodeTargets; }
+
+	private:
+		int32_t											  SpatialId;
+		int32_t											  TemporalId;
+		TRefCountPtr<FEpicRtcDecodeTargetIndicationArray> DecodeTargetIndications;
+		TRefCountPtr<FEpicRtcInt32Array>				  FrameDiffs;
+		TRefCountPtr<FEpicRtcInt32Array>				  ChainDiffs;
+		TRefCountPtr<FEpicRtcCodecBufferUsageArray>		  EncoderBuffers;
+		TRefCountPtr<FEpicRtcBoolArray>					  PartOfChain;
+		TRefCountPtr<FEpicRtcBoolArray>					  ActiveDecodeTargets;
+
+	public:
+		// Begin EpicRtcRefCountInterface
+		EPICRTC_REFCOUNT_INTERFACE_IN_PLACE
+		// End EpicRtcRefCountInterface
+	};
+
+	class PIXELSTREAMING2RTC_API FEpicRtcFrameDependencyStructure : public EpicRtcFrameDependencyStructure
+	{
+	public:
+		FEpicRtcFrameDependencyStructure(const FFrameDependencyStructure& FrameDependencyStructure)
+			: StructureId(FrameDependencyStructure.StructureId)
+			, NumDecodeTargets(FrameDependencyStructure.NumDecodeTargets)
+			, NumChains(FrameDependencyStructure.NumChains)
+			, DecodeTargetProtectedByChain(MakeRefCount<FEpicRtcInt32Array>(FrameDependencyStructure.DecodeTargetProtectedByChain))
+			, Resolutions(MakeRefCount<FEpicRtcVideoResolutionArray>(FrameDependencyStructure.Resolutions))
+		{
+			TArray<EpicRtcGenericFrameInfoInterface*> GenericFrameInfoArray;
+			GenericFrameInfoArray.SetNum(FrameDependencyStructure.Templates.Num());
+
+			for (size_t i = 0; i < FrameDependencyStructure.Templates.Num(); i++)
+			{
+				FGenericFrameInfo GenericFrameInfo;
+				GenericFrameInfo.SpatialId = FrameDependencyStructure.Templates[i].SpatialId;
+				GenericFrameInfo.TemporalId = FrameDependencyStructure.Templates[i].TemporalId;
+				GenericFrameInfo.DecodeTargetIndications = FrameDependencyStructure.Templates[i].DecodeTargetIndications;
+				GenericFrameInfo.FrameDiffs = FrameDependencyStructure.Templates[i].FrameDiffs;
+				GenericFrameInfo.ChainDiffs = FrameDependencyStructure.Templates[i].ChainDiffs;
+
+				GenericFrameInfoArray[i] = new FEpicRtcGenericFrameInfo(GenericFrameInfo);
+			}
+
+			Templates = MakeRefCount<FEpicRtcGenericFrameInfoArray>(GenericFrameInfoArray);
+		}
+
+		virtual ~FEpicRtcFrameDependencyStructure() = default;
+
+		virtual int32_t								   GetStructureId() override { return StructureId; }
+		virtual int32_t								   GetNumDecodeTargets() override { return NumDecodeTargets; }
+		virtual int32_t								   GetNumChains() override { return NumChains; }
+		virtual EpicRtcInt32ArrayInterface*			   GetDecodeTargetProtectedByChain() override { return DecodeTargetProtectedByChain; }
+		virtual EpicRtcVideoResolutionArrayInterface*  GetResolutions() override { return Resolutions; }
+		virtual EpicRtcGenericFrameInfoArrayInterface* GetTemplates() override { return Templates; }
+
+		friend bool operator==(FEpicRtcFrameDependencyStructure& Lhs, FEpicRtcFrameDependencyStructure& Rhs)
+		{
+			TArray<int32_t> LhsDecodeTargetProtectedByChain(Lhs.GetDecodeTargetProtectedByChain()->Get(), Lhs.GetDecodeTargetProtectedByChain()->Size());
+			TArray<int32_t> RhsDecodeTargetProtectedByChain(Rhs.GetDecodeTargetProtectedByChain()->Get(), Rhs.GetDecodeTargetProtectedByChain()->Size());
+
+			TArray<EpicRtcVideoResolution> LhsResolutions(Lhs.GetResolutions()->Get(), Lhs.GetResolutions()->Size());
+			TArray<EpicRtcVideoResolution> RhsResolutions(Rhs.GetResolutions()->Get(), Rhs.GetResolutions()->Size());
+
+			TArray<EpicRtcGenericFrameInfoInterface*> LhsTemplates(Lhs.GetTemplates()->Get(), Lhs.GetTemplates()->Size());
+			TArray<EpicRtcGenericFrameInfoInterface*> RhsTemplates(Rhs.GetTemplates()->Get(), Rhs.GetTemplates()->Size());
+
+			return Lhs.NumDecodeTargets == Rhs.NumDecodeTargets
+				&& Lhs.NumChains == Rhs.NumChains
+				&& LhsDecodeTargetProtectedByChain == RhsDecodeTargetProtectedByChain
+				&& LhsResolutions == RhsResolutions
+				&& LhsTemplates == RhsTemplates;
+		}
+
+	private:
+		int											StructureId;
+		int											NumDecodeTargets;
+		int											NumChains;
+		TRefCountPtr<FEpicRtcInt32Array>			DecodeTargetProtectedByChain;
+		TRefCountPtr<FEpicRtcVideoResolutionArray>	Resolutions;
+		TRefCountPtr<FEpicRtcGenericFrameInfoArray> Templates;
+
+	public:
+		// Begin EpicRtcRefCountInterface
+		EPICRTC_REFCOUNT_INTERFACE_IN_PLACE
+		// End EpicRtcRefCountInterface
+	};
+} // namespace UE::PixelStreaming2
\ No newline at end of file
diff --git a/Engine/Plugins/Media/PixelStreaming2/Source/PixelStreaming2RTC/Internal/EpicRtcVideoDecoderInitializer.h b/Engine/Plugins/Media/PixelStreaming2/Source/PixelStreaming2RTC/Internal/EpicRtcVideoDecoderInitializer.h
new file mode 100644
--- /dev/null
+++ b/Engine/Plugins/Media/PixelStreaming2/Source/PixelStreaming2RTC/Internal/EpicRtcVideoDecoderInitializer.h
@@ -0,0 +1,35 @@
+// Copyright Epic Games, Inc. All Rights Reserved.
+
+#pragma once
+
+#include "EpicRtcVideoCommon.h"
+#include "Video/VideoConfig.h"
+
+#include "epic_rtc/core/video/video_decoder.h"
+#include "epic_rtc_helper/memory/ref_count_impl_helper.h"
+
+namespace UE::PixelStreaming2
+{
+
+	class PIXELSTREAMING2RTC_API FEpicRtcVideoDecoderInitializer : public EpicRtcVideoDecoderInitializerInterface
+	{
+	public:
+		FEpicRtcVideoDecoderInitializer() = default;
+		virtual ~FEpicRtcVideoDecoderInitializer() = default;
+
+		// Begin EpicRtcVideoDecoderInitializerInterface
+		virtual void								 CreateDecoder(EpicRtcVideoCodecInfoInterface* CodecInfo, EpicRtcVideoDecoderInterface** OutDecoder) override;
+		virtual EpicRtcStringView					 GetName() override;
+		virtual EpicRtcVideoCodecInfoArrayInterface* GetSupportedCodecs() override;
+		// End EpicRtcVideoDecoderInitializerInterface
+
+	private:
+		TMap<EVideoCodec, TArray<TRefCountPtr<EpicRtcVideoCodecInfoInterface>>> CreateSupportedDecoderMap();
+
+	public:
+		// Begin EpicRtcRefCountInterface
+		EPICRTC_REFCOUNT_INTERFACE_IN_PLACE
+		// End EpicRtcRefCountInterface
+	};
+
+} // namespace UE::PixelStreaming2
\ No newline at end of file
diff --git a/Engine/Plugins/Media/PixelStreaming2/Source/PixelStreaming2RTC/Internal/EpicRtcVideoEncoderInitializer.h b/Engine/Plugins/Media/PixelStreaming2/Source/PixelStreaming2RTC/Internal/EpicRtcVideoEncoderInitializer.h
new file mode 100644
--- /dev/null
+++ b/Engine/Plugins/Media/PixelStreaming2/Source/PixelStreaming2RTC/Internal/EpicRtcVideoEncoderInitializer.h
@@ -0,0 +1,35 @@
+// Copyright Epic Games, Inc. All Rights Reserved.
+
+#pragma once
+
+#include "EpicRtcVideoCommon.h"
+#include "Video/VideoConfig.h"
+
+#include "epic_rtc/core/video/video_encoder.h"
+#include "epic_rtc_helper/memory/ref_count_impl_helper.h"
+
+namespace UE::PixelStreaming2
+{
+
+	class PIXELSTREAMING2RTC_API FEpicRtcVideoEncoderInitializer : public EpicRtcVideoEncoderInitializerInterface
+	{
+	public:
+		FEpicRtcVideoEncoderInitializer() = default;
+		virtual ~FEpicRtcVideoEncoderInitializer() = default;
+
+		// Begin EpicRtcVideoEncoderInitializerInterface
+		virtual void								 CreateEncoder(EpicRtcVideoCodecInfoInterface* CodecInfo, EpicRtcVideoEncoderInterface** OutEncoder) override;
+		virtual EpicRtcStringView					 GetName() override;
+		virtual EpicRtcVideoCodecInfoArrayInterface* GetSupportedCodecs() override;
+		// End EpicRtcVideoEncoderInitializerInterface
+
+	private:
+		TMap<EVideoCodec, TArray<TRefCountPtr<EpicRtcVideoCodecInfoInterface>>> CreateSupportedEncoderMap();
+
+	public:
+		// Begin EpicRtcRefCountInterface
+		EPICRTC_REFCOUNT_INTERFACE_IN_PLACE
+		// End EpicRtcRefCountInterface
+	};
+
+} // namespace UE::PixelStreaming2
\ No newline at end of file
diff --git a/Engine/Plugins/Media/PixelStreaming2/Source/PixelStreaming2RTC/Internal/EpicRtcVideoTrackObserver.h b/Engine/Plugins/Media/PixelStreaming2/Source/PixelStreaming2RTC/Internal/EpicRtcVideoTrackObserver.h
new file mode 100644
--- /dev/null
+++ b/Engine/Plugins/Media/PixelStreaming2/Source/PixelStreaming2RTC/Internal/EpicRtcVideoTrackObserver.h
@@ -0,0 +1,39 @@
+// Copyright Epic Games, Inc. All Rights Reserved.
+
+#pragma once
+
+#include "Templates/SharedPointer.h"
+
+#include "epic_rtc/core/video/video_track_observer.h"
+#include "epic_rtc_helper/memory/ref_count_impl_helper.h"
+
+namespace UE::PixelStreaming2
+{
+	class FEpicRtcManager;
+
+	class PIXELSTREAMING2RTC_API FEpicRtcVideoTrackObserver : public EpicRtcVideoTrackObserverInterface
+	{
+	public:
+		FEpicRtcVideoTrackObserver(TWeakPtr<FEpicRtcManager> Manager);
+		virtual ~FEpicRtcVideoTrackObserver() = default;
+
+	private:
+		// Begin EpicRtcVideoTrackObserverInterface
+		virtual void		OnVideoTrackMuted(EpicRtcVideoTrackInterface* VideoTrack, EpicRtcBool bIsMuted) override;
+		virtual void		OnVideoTrackFrame(EpicRtcVideoTrackInterface* VideoTrack, const EpicRtcVideoFrame& Frame) override;
+		virtual void		OnVideoTrackRemoved(EpicRtcVideoTrackInterface* VideoTrack) override;
+		virtual void		OnVideoTrackState(EpicRtcVideoTrackInterface* VideoTrack, const EpicRtcTrackState State) override;
+		virtual void		OnVideoTrackEncodedFrame(EpicRtcVideoTrackInterface*, const EpicRtcEncodedVideoFrame&) override {};
+		virtual EpicRtcBool Enabled() const override { return true; };
+		// End EpicRtcVideoTrackObserverInterface
+
+	public:
+		// Begin EpicRtcRefCountInterface
+		EPICRTC_REFCOUNT_INTERFACE_IN_PLACE
+		// End EpicRtcRefCountInterface
+
+	private:
+		TWeakPtr<FEpicRtcManager> Manager;
+	};
+
+} // namespace UE::PixelStreaming2
\ No newline at end of file
diff --git a/Engine/Plugins/Media/PixelStreaming2/Source/PixelStreaming2RTC/Internal/EpicRtcVideoTrackObserverFactory.h b/Engine/Plugins/Media/PixelStreaming2/Source/PixelStreaming2RTC/Internal/EpicRtcVideoTrackObserverFactory.h
new file mode 100644
--- /dev/null
+++ b/Engine/Plugins/Media/PixelStreaming2/Source/PixelStreaming2RTC/Internal/EpicRtcVideoTrackObserverFactory.h
@@ -0,0 +1,34 @@
+// Copyright Epic Games, Inc. All Rights Reserved.
+
+#pragma once
+
+#include "Templates/SharedPointer.h"
+
+#include "epic_rtc/core/video/video_track_observer.h"
+#include "epic_rtc_helper/memory/ref_count_impl_helper.h"
+
+namespace UE::PixelStreaming2
+{
+	class FEpicRtcManager;
+
+	class PIXELSTREAMING2RTC_API FEpicRtcVideoTrackObserverFactory : public EpicRtcVideoTrackObserverFactoryInterface
+	{
+	public:
+		FEpicRtcVideoTrackObserverFactory(TWeakPtr<FEpicRtcManager> Manager);
+		virtual ~FEpicRtcVideoTrackObserverFactory() = default;
+
+	public:
+		// Begin EpicRtcVideoTrackObserverFactoryInterface
+		virtual EpicRtcErrorCode CreateVideoTrackObserver(const EpicRtcStringView ParticipantId, const EpicRtcStringView VideoTrackId, EpicRtcVideoTrackObserverInterface** OutVideoTrackObserver) override;
+		// End EpicRtcVideoTrackObserverFactoryInterface
+
+	public:
+		// Begin EpicRtcRefCountInterface
+		EPICRTC_REFCOUNT_INTERFACE_IN_PLACE
+		// End EpicRtcRefCountInterface
+
+	private:
+		TWeakPtr<FEpicRtcManager> Manager;
+	};
+
+} // namespace UE::PixelStreaming2
\ No newline at end of file
diff --git a/Engine/Plugins/Media/PixelStreaming2/Source/PixelStreaming2RTC/Internal/EpicRtcWebsocket.h b/Engine/Plugins/Media/PixelStreaming2/Source/PixelStreaming2RTC/Internal/EpicRtcWebsocket.h
new file mode 100644
--- /dev/null
+++ b/Engine/Plugins/Media/PixelStreaming2/Source/PixelStreaming2RTC/Internal/EpicRtcWebsocket.h
@@ -0,0 +1,66 @@
+// Copyright Epic Games, Inc. All Rights Reserved.
+
+#pragma once
+
+#include "Engine/EngineTypes.h"
+#include "Tickable.h"
+
+#include "epic_rtc/plugins/signalling/websocket.h"
+#include "epic_rtc_helper/memory/ref_count_impl_helper.h"
+
+class IWebSocket;
+
+namespace UE::PixelStreaming2
+{
+	class PIXELSTREAMING2RTC_API FEpicRtcWebsocket : public EpicRtcWebsocketInterface, public FTickableGameObject
+	{
+	public:
+		FEpicRtcWebsocket(bool bKeepAlive = true, TSharedPtr<IWebSocket> WebSocket = nullptr);
+		virtual ~FEpicRtcWebsocket() = default;
+
+		// Begin EpicRtcWebsocketInterface
+		virtual EpicRtcBool Connect(EpicRtcStringView Url, EpicRtcWebsocketObserverInterface* Observer) override;
+		virtual void		Disconnect(const EpicRtcStringView Reason) override;
+		virtual void		Send(EpicRtcStringView Message) override;
+		// End EpicRtcWebsocketInterface
+
+		// Begin FTickableGameObject
+		virtual void		Tick(float DeltaTime) override;
+		virtual bool		IsTickableInEditor() const override { return true; }
+		virtual bool		IsTickableWhenPaused() const override { return true; }
+		FORCEINLINE TStatId GetStatId() const { RETURN_QUICK_DECLARE_CYCLE_STAT(EpicRtcWebSocket, STATGROUP_Tickables); }
+		// End FTickableGameObject
+
+	private:
+		void OnConnected();
+		void OnConnectionError(const FString& Error);
+		void OnClosed(int32 StatusCode, const FString& Reason, bool bWasClean);
+		void OnMessage(const FString& Msg);
+		void OnBinaryMessage(const uint8* Data, int32 Length, bool bIsLastFragment);
+		void Reconnect();
+		void KeepAlive();
+
+	private:
+		FDelegateHandle OnConnectedHandle;
+		FDelegateHandle OnConnectionErrorHandle;
+		FDelegateHandle OnClosedHandle;
+		FDelegateHandle OnMessageHandle;
+		FDelegateHandle OnBinaryMessageHandle;
+
+	private:
+		FString											Url;
+		TSharedPtr<IWebSocket>							WebSocket;
+		TRefCountPtr<EpicRtcWebsocketObserverInterface> Observer;
+		uint64											LastKeepAliveCycles = 0;
+		uint64											LastReconnectCycles = 0;
+		bool											bSendKeepAlive = false;
+		bool											bReconnectOnError = false;
+		bool											bCloseRequested = false;
+
+	public:
+		// Begin EpicRtcRefCountInterface
+		EPICRTC_REFCOUNT_INTERFACE_IN_PLACE
+		// End EpicRtcRefCountInterface
+	};
+
+} // namespace UE::PixelStreaming2
\ No newline at end of file
diff --git a/Engine/Plugins/Media/PixelStreaming2/Source/PixelStreaming2RTC/Private/EpicRtcAudioCapturer.cpp b/Engine/Plugins/Media/PixelStreaming2/Source/PixelStreaming2RTC/Private/EpicRtcAudioCapturer.cpp
new file mode 100644
--- /dev/null
+++ b/Engine/Plugins/Media/PixelStreaming2/Source/PixelStreaming2RTC/Private/EpicRtcAudioCapturer.cpp
@@ -0,0 +1,54 @@
+// Copyright Epic Games, Inc. All Rights Reserved.
+
+#include "EpicRtcAudioCapturer.h"
+
+#include "AudioDevice.h"
+#include "Engine/Engine.h"
+#include "Logging.h"
+#include "Misc/CoreDelegates.h"
+#include "PixelStreaming2PluginSettings.h"
+#include "Sound/SampleBufferIO.h"
+
+namespace UE::PixelStreaming2
+{
+	TSharedPtr<FEpicRtcAudioCapturer> FEpicRtcAudioCapturer::Create()
+	{
+		TSharedPtr<FEpicRtcAudioCapturer> AudioMixingCapturer(new FEpicRtcAudioCapturer());
+
+		FAudioDeviceManagerDelegates::OnAudioDeviceCreated.AddSP(AudioMixingCapturer.ToSharedRef(), &FEpicRtcAudioCapturer::CreateAudioProducer);
+		FAudioDeviceManagerDelegates::OnAudioDeviceDestroyed.AddSP(AudioMixingCapturer.ToSharedRef(), &FEpicRtcAudioCapturer::RemoveAudioProducer);
+
+		if (UPixelStreaming2PluginSettings::FDelegates* Delegates = UPixelStreaming2PluginSettings::Delegates())
+		{
+			Delegates->OnDebugDumpAudioChanged.AddSP(AudioMixingCapturer.ToSharedRef(), &FEpicRtcAudioCapturer::OnDebugDumpAudioChanged);
+
+			TWeakPtr<FEpicRtcAudioCapturer> WeakAudioMixingCapturer = AudioMixingCapturer;
+			FCoreDelegates::OnEnginePreExit.AddLambda([WeakAudioMixingCapturer]() {
+				if (TSharedPtr<FEpicRtcAudioCapturer> AudioMixingCapturer = WeakAudioMixingCapturer.Pin())
+				{
+					AudioMixingCapturer->OnEnginePreExit();
+				}
+			});
+		}
+
+		return AudioMixingCapturer;
+	}
+
+	void FEpicRtcAudioCapturer::PushAudio(const float* AudioData, int32 InNumSamples, int32 InNumChannels, const int32 InSampleRate)
+	{
+		Audio::TSampleBuffer<int16_t> Buffer(AudioData, InNumSamples, InNumChannels, SampleRate);
+		RecordingBuffer.Append(Buffer.GetData(), Buffer.GetNumSamples());
+
+		const int32	 SamplesPer10Ms = NumChannels * SampleRate * 0.01f;
+		const size_t BytesPerFrame = NumChannels * sizeof(int16_t);
+
+		// Feed in 10ms chunks
+		while (RecordingBuffer.Num() > SamplesPer10Ms)
+		{
+			OnAudioBuffer.Broadcast(RecordingBuffer.GetData(), SamplesPer10Ms, NumChannels, SampleRate);
+
+			// Remove 10ms of samples from the recording buffer now it is submitted
+			RecordingBuffer.RemoveAt(0, SamplesPer10Ms, EAllowShrinking::No);
+		}
+	}
+} // namespace UE::PixelStreaming2
\ No newline at end of file
diff --git a/Engine/Plugins/Media/PixelStreaming2/Source/PixelStreaming2RTC/Private/EpicRtcAudioSource.cpp b/Engine/Plugins/Media/PixelStreaming2/Source/PixelStreaming2RTC/Private/EpicRtcAudioSource.cpp
new file mode 100644
--- /dev/null
+++ b/Engine/Plugins/Media/PixelStreaming2/Source/PixelStreaming2RTC/Private/EpicRtcAudioSource.cpp
@@ -0,0 +1,53 @@
+// Copyright Epic Games, Inc. All Rights Reserved.
+
+#include "EpicRtcAudioSource.h"
+
+#include "EpicRtcAudioCapturer.h"
+#include "PixelStreaming2RTCModule.h"
+
+namespace UE::PixelStreaming2
+{
+
+	TSharedPtr<FEpicRtcAudioSource> FEpicRtcAudioSource::Create(TRefCountPtr<EpicRtcAudioTrackInterface> InTrack)
+	{
+		TSharedPtr<FEpicRtcAudioSource> AudioTrack = TSharedPtr<FEpicRtcAudioSource>(new FEpicRtcAudioSource(InTrack));
+
+		TSharedPtr<FEpicRtcAudioCapturer> Capturer = FPixelStreaming2RTCModule::GetModule()->GetAudioCapturer();
+		Capturer->OnAudioBuffer.AddSP(AudioTrack.ToSharedRef(), &FEpicRtcAudioSource::OnAudioBuffer);
+
+		return AudioTrack;
+	}
+
+	FEpicRtcAudioSource::FEpicRtcAudioSource(TRefCountPtr<EpicRtcAudioTrackInterface> InTrack)
+	{
+		Track = InTrack;
+	}
+
+	void FEpicRtcAudioSource::OnAudioBuffer(const int16_t* AudioData, int32 NumSamples, int32 NumChannels, const int32 SampleRate)
+	{
+		if (!Track || bIsMuted)
+		{
+			return;
+		}
+
+		const uint32_t NumFrames = NumSamples / NumChannels;
+		EpicRtcAudioFrame AudioFrame
+		{
+			._data = const_cast<int16_t*>(AudioData),
+			._length = NumFrames,
+			._timestamp = 0,
+			._format =
+			{
+				._sampleRate = (uint32_t)SampleRate,
+				._numChannels = (uint32_t)NumChannels,
+				._parameters = nullptr
+			}
+		};
+
+		// Because UE handles all audio processing, we can bypass the ADM.
+		// This also has the added benefit of increasing audio quality
+		Track->PushFrame(AudioFrame, true);
+	}
+
+} // namespace UE::PixelStreaming2
+
diff --git a/Engine/Plugins/Media/PixelStreaming2/Source/PixelStreaming2RTC/Private/EpicRtcAudioSource.h b/Engine/Plugins/Media/PixelStreaming2/Source/PixelStreaming2RTC/Private/EpicRtcAudioSource.h
new file mode 100644
--- /dev/null
+++ b/Engine/Plugins/Media/PixelStreaming2/Source/PixelStreaming2RTC/Private/EpicRtcAudioSource.h
@@ -0,0 +1,23 @@
+// Copyright Epic Games, Inc. All Rights Reserved.
+
+#pragma once
+
+#include "AudioSource.h"
+#include "EpicRtcAudioTrack.h"
+#include "Templates/SharedPointer.h"
+
+#include "epic_rtc/core/audio/audio_track.h"
+
+namespace UE::PixelStreaming2
+{
+	class FEpicRtcAudioSource : public FAudioSource, public FEpicRtcAudioTrack
+	{
+	public:
+		static TSharedPtr<FEpicRtcAudioSource> Create(TRefCountPtr<EpicRtcAudioTrackInterface> InTrack);
+
+		virtual void OnAudioBuffer(const int16_t* AudioData, int32 NumSamples, int32 NumChannels, const int32 SampleRate);
+
+	private:
+		FEpicRtcAudioSource(TRefCountPtr<EpicRtcAudioTrackInterface> InTrack);
+	};
+} // namespace UE::PixelStreaming2
\ No newline at end of file
diff --git a/Engine/Plugins/Media/PixelStreaming2/Source/PixelStreaming2RTC/Private/EpicRtcAudioTrack.cpp b/Engine/Plugins/Media/PixelStreaming2/Source/PixelStreaming2RTC/Private/EpicRtcAudioTrack.cpp
new file mode 100644
--- /dev/null
+++ b/Engine/Plugins/Media/PixelStreaming2/Source/PixelStreaming2RTC/Private/EpicRtcAudioTrack.cpp
@@ -0,0 +1,16 @@
+// Copyright Epic Games, Inc. All Rights Reserved.
+
+#include "EpicRtcAudioTrack.h"
+	
+namespace UE::PixelStreaming2
+{
+	EpicRtcStringView FEpicRtcAudioTrack::GetTrackId() const
+	{
+		if (!Track)
+		{
+			return EpicRtcStringView { ._ptr = nullptr, ._length = 0 };
+		}
+
+		return Track->GetId();
+	}
+} // namespace UE::PixelStreaming2
\ No newline at end of file
diff --git a/Engine/Plugins/Media/PixelStreaming2/Source/PixelStreaming2RTC/Private/EpicRtcAudioTrackObserver.cpp b/Engine/Plugins/Media/PixelStreaming2/Source/PixelStreaming2RTC/Private/EpicRtcAudioTrackObserver.cpp
new file mode 100644
--- /dev/null
+++ b/Engine/Plugins/Media/PixelStreaming2/Source/PixelStreaming2RTC/Private/EpicRtcAudioTrackObserver.cpp
@@ -0,0 +1,55 @@
+// Copyright Epic Games, Inc. All Rights Reserved.
+
+#include "EpicRtcAudioTrackObserver.h"
+
+#include "EpicRtcManager.h"
+
+namespace UE::PixelStreaming2
+{
+
+	FEpicRtcAudioTrackObserver::FEpicRtcAudioTrackObserver(TWeakPtr<FEpicRtcManager> Manager)
+		: Manager(Manager)
+	{
+	}
+
+	void FEpicRtcAudioTrackObserver::OnAudioTrackMuted(EpicRtcAudioTrackInterface* AudioTrack, EpicRtcBool bIsMuted)
+	{
+		if (!Manager.IsValid())
+		{
+			return;
+		}
+
+		Manager.Pin()->OnAudioTrackMuted.Broadcast(AudioTrack, bIsMuted);
+	}
+
+	void FEpicRtcAudioTrackObserver::OnAudioTrackFrame(EpicRtcAudioTrackInterface* AudioTrack, const EpicRtcAudioFrame& Frame)
+	{
+		if (!Manager.IsValid())
+		{
+			return;
+		}
+
+		Manager.Pin()->OnAudioTrackFrame.Broadcast(AudioTrack, Frame);
+	}
+
+	void FEpicRtcAudioTrackObserver::OnAudioTrackRemoved(EpicRtcAudioTrackInterface* AudioTrack)
+	{
+		if (!Manager.IsValid())
+		{
+			return;
+		}
+
+		Manager.Pin()->OnAudioTrackRemoved.Broadcast(AudioTrack);
+	}
+
+	void FEpicRtcAudioTrackObserver::OnAudioTrackState(EpicRtcAudioTrackInterface* AudioTrack, const EpicRtcTrackState State)
+	{
+		if (!Manager.IsValid())
+		{
+			return;
+		}
+
+		Manager.Pin()->OnAudioTrackState.Broadcast(AudioTrack, State);
+	}
+
+} // namespace UE::PixelStreaming2
\ No newline at end of file
diff --git a/Engine/Plugins/Media/PixelStreaming2/Source/PixelStreaming2RTC/Private/EpicRtcDataTrackObserver.cpp b/Engine/Plugins/Media/PixelStreaming2/Source/PixelStreaming2RTC/Private/EpicRtcDataTrackObserver.cpp
new file mode 100644
--- /dev/null
+++ b/Engine/Plugins/Media/PixelStreaming2/Source/PixelStreaming2RTC/Private/EpicRtcDataTrackObserver.cpp
@@ -0,0 +1,35 @@
+// Copyright Epic Games, Inc. All Rights Reserved.
+
+#include "EpicRtcDataTrackObserver.h"
+
+#include "EpicRtcManager.h"
+
+namespace UE::PixelStreaming2
+{
+
+	FEpicRtcDataTrackObserver::FEpicRtcDataTrackObserver(TWeakPtr<FEpicRtcManager> Manager)
+		: Manager(Manager)
+	{
+	}
+
+	void FEpicRtcDataTrackObserver::OnDataTrackState(EpicRtcDataTrackInterface* DataTrack, const EpicRtcTrackState State)
+	{
+		if (!Manager.IsValid())
+		{
+			return;
+		}
+
+		Manager.Pin()->OnDataTrackState.Broadcast(DataTrack, State);
+	}
+
+	void FEpicRtcDataTrackObserver::OnDataTrackMessage(EpicRtcDataTrackInterface* DataTrack)
+	{
+		if (!Manager.IsValid())
+		{
+			return;
+		}
+
+		Manager.Pin()->OnDataTrackMessage.Broadcast(DataTrack);
+	}
+
+} // namespace UE::PixelStreaming2
\ No newline at end of file
diff --git a/Engine/Plugins/Media/PixelStreaming2/Source/PixelStreaming2RTC/Private/EpicRtcRoomObserver.cpp b/Engine/Plugins/Media/PixelStreaming2/Source/PixelStreaming2RTC/Private/EpicRtcRoomObserver.cpp
new file mode 100644
--- /dev/null
+++ b/Engine/Plugins/Media/PixelStreaming2/Source/PixelStreaming2RTC/Private/EpicRtcRoomObserver.cpp
@@ -0,0 +1,109 @@
+// Copyright Epic Games, Inc. All Rights Reserved.
+
+#include "EpicRtcRoomObserver.h"
+
+#include "EpicRtcManager.h"
+
+namespace UE::PixelStreaming2
+{
+
+	FEpicRtcRoomObserver::FEpicRtcRoomObserver(TWeakPtr<FEpicRtcManager> Manager)
+		: Manager(Manager)
+	{
+	}
+
+	void FEpicRtcRoomObserver::OnRoomStateUpdate(const EpicRtcRoomState State)
+	{
+		if (!Manager.IsValid())
+		{
+			return;
+		}
+
+		Manager.Pin()->OnRoomStateUpdate.Broadcast(State);
+	}
+
+	void FEpicRtcRoomObserver::OnRoomJoinedUpdate(EpicRtcParticipantInterface* Participant)
+	{
+		if (!Manager.IsValid())
+		{
+			return;
+		}
+
+		Manager.Pin()->OnRoomJoinedUpdate.Broadcast(Participant);
+	}
+
+	void FEpicRtcRoomObserver::OnRoomLeftUpdate(const EpicRtcStringView ParticipantId)
+	{
+		if (!Manager.IsValid())
+		{
+			return;
+		}
+
+		Manager.Pin()->OnRoomLeftUpdate.Broadcast(ParticipantId);
+	}
+
+	void FEpicRtcRoomObserver::OnAudioTrackUpdate(EpicRtcParticipantInterface* Participant, EpicRtcAudioTrackInterface* AudioTrack)
+	{
+		if (!Manager.IsValid())
+		{
+			return;
+		}
+
+		Manager.Pin()->OnAudioTrackUpdate.Broadcast(Participant, AudioTrack);
+	}
+
+	void FEpicRtcRoomObserver::OnVideoTrackUpdate(EpicRtcParticipantInterface* Participant, EpicRtcVideoTrackInterface* VideoTrack)
+	{
+		if (!Manager.IsValid())
+		{
+			return;
+		}
+
+		Manager.Pin()->OnVideoTrackUpdate.Broadcast(Participant, VideoTrack);
+	}
+
+	void FEpicRtcRoomObserver::OnDataTrackUpdate(EpicRtcParticipantInterface* Participant, EpicRtcDataTrackInterface* DataTrack)
+	{
+		if (!Manager.IsValid())
+		{
+			return;
+		}
+
+		Manager.Pin()->OnDataTrackUpdate.Broadcast(Participant, DataTrack);
+	}
+
+	[[nodiscard]] EpicRtcSdpInterface* FEpicRtcRoomObserver::OnLocalSdpUpdate(EpicRtcParticipantInterface* Participant, EpicRtcSdpInterface* Sdp)
+	{
+		if (!Manager.IsValid())
+		{
+			return nullptr;
+		}
+
+		Manager.Pin()->OnLocalSdpUpdate.Broadcast(Participant, Sdp);
+
+		return nullptr;
+	}
+
+	[[nodiscard]] EpicRtcSdpInterface* FEpicRtcRoomObserver::OnRemoteSdpUpdate(EpicRtcParticipantInterface* Participant, EpicRtcSdpInterface* Sdp)
+	{
+		if (!Manager.IsValid())
+		{
+			return nullptr;
+		}
+
+		Manager.Pin()->OnRemoteSdpUpdate.Broadcast(Participant, Sdp);
+
+		return nullptr;
+	}
+
+	void FEpicRtcRoomObserver::OnRoomErrorUpdate(const EpicRtcErrorCode Error)
+	{
+		if (!Manager.IsValid())
+		{
+			return;
+		}
+
+		Manager.Pin()->OnRoomErrorUpdate.Broadcast(Error);
+	}
+
+} // namespace UE::PixelStreaming2
\ No newline at end of file
diff --git a/Engine/Plugins/Media/PixelStreaming2/Source/PixelStreaming2RTC/Private/EpicRtcSessionObserver.cpp b/Engine/Plugins/Media/PixelStreaming2/Source/PixelStreaming2RTC/Private/EpicRtcSessionObserver.cpp
new file mode 100644
--- /dev/null
+++ b/Engine/Plugins/Media/PixelStreaming2/Source/PixelStreaming2RTC/Private/EpicRtcSessionObserver.cpp
@@ -0,0 +1,45 @@
+// Copyright Epic Games, Inc. All Rights Reserved.
+
+#include "EpicRtcSessionObserver.h"
+
+#include "EpicRtcManager.h"
+
+namespace UE::PixelStreaming2
+{
+
+	FEpicRtcSessionObserver::FEpicRtcSessionObserver(TWeakPtr<FEpicRtcManager> Manager)
+		: Manager(Manager)
+	{
+	}
+
+	void FEpicRtcSessionObserver::OnSessionStateUpdate(const EpicRtcSessionState State)
+	{
+		if (!Manager.IsValid())
+		{
+			return;
+		}
+
+		Manager.Pin()->OnSessionStateUpdate.Broadcast(State);
+	}
+
+	void FEpicRtcSessionObserver::OnSessionErrorUpdate(const EpicRtcErrorCode Error)
+	{
+		if (!Manager.IsValid())
+		{
+			return;
+		}
+
+		Manager.Pin()->OnSessionErrorUpdate.Broadcast(Error);
+	}
+
+	void FEpicRtcSessionObserver::OnSessionRoomsAvailableUpdate(EpicRtcStringArrayInterface* RoomsList)
+	{
+		if (!Manager.IsValid())
+		{
+			return;
+		}
+
+		Manager.Pin()->OnSessionRoomsAvailableUpdate.Broadcast(RoomsList);
+	}
+
+} // namespace UE::PixelStreaming2
