[PS2] PR #12718: Refactor PixelStreaming2 plugin to enable other streaming protocols

diff --git a/Engine/Plugins/Media/PixelStreaming2/Source/PixelStreaming2RTC/Private/EpicRtcStreamer.cpp b/Engine/Plugins/Media/PixelStreaming2/Source/PixelStreaming2RTC/Private/EpicRtcStreamer.cpp
new file mode 100644
--- /dev/null
+++ b/Engine/Plugins/Media/PixelStreaming2/Source/PixelStreaming2RTC/Private/EpicRtcStreamer.cpp
@@ -0,0 +1,1639 @@
+// Copyright Epic Games, Inc. All Rights Reserved.
+
+#include "EpicRtcStreamer.h"
+
+#include "PixelStreaming2PluginSettings.h"
+#include "PixelStreaming2Delegates.h"
+#include "EpicRtcWebsocket.h"
+#include "EpicRtcWebsocketFactory.h"
+#include "TextureResource.h"
+#include "WebSocketsModule.h"
+#include "Logging.h"
+#include "Stats.h"
+#include "PixelStreaming2StatNames.h"
+#include "Async/Async.h"
+#include "Engine/Texture2D.h"
+#include "RTCStatsCollector.h"
+#include "Framework/Application/SlateApplication.h"
+#include "Dom/JsonValue.h"
+#include "Serialization/JsonWriter.h"
+#include "Serialization/JsonSerializer.h"
+#include "Dom/JsonObject.h"
+#include "PixelStreaming2RTCModule.h"
+#include "EpicRtcAudioTrack.h"
+#include "EpicRtcDataTrack.h"
+#include "IPixelStreaming2InputModule.h"
+#include "CoreGlobals.h"
+#include "UtilsString.h"
+#include "UtilsCodecs.h"
+
+namespace UE::PixelStreaming2
+{
+	FEpicRtcStreamer::FEpicRtcStreamer(const FString& InStreamerId, TRefCountPtr<EpicRtcConferenceInterface> Conference)
+		: StreamerId(InStreamerId)
+		, InputHandler(IPixelStreaming2InputModule::Get().CreateInputHandler())
+		, Players(new TThreadSafeMap<FString, FPlayerContext>())
+		, VideoCapturer(FEpicRtcVideoCapturer::Create())
+		, VideoSourceGroup(FVideoSourceGroup::Create(VideoCapturer))
+		, FreezeFrame(FFreezeFrame::Create(Players, VideoCapturer, InputHandler))
+		, EpicRtcManager(MakeShared<FEpicRtcManager>())
+	{
+		ReconnectTimer = MakeShared<FStreamerReconnectTimer>();
+
+		InputHandler->SetElevatedCheck([this](FString PlayerId) {
+			return GetEnumFromCVar<EInputControllerMode>(UPixelStreaming2PluginSettings::CVarInputController) == EInputControllerMode::Any
+				|| InputControllingId == INVALID_PLAYER_ID
+				|| PlayerId == InputControllingId;
+		});
+
+		EpicRtcManager->EpicRtcConference = Conference;
+
+		EpicRtcManager->SessionObserver = MakeRefCount<FEpicRtcSessionObserver>(EpicRtcManager);
+		EpicRtcManager->RoomObserver = MakeRefCount<FEpicRtcRoomObserver>(EpicRtcManager);
+
+		EpicRtcManager->AudioTrackObserverFactory = MakeRefCount<FEpicRtcAudioTrackObserverFactory>(EpicRtcManager);
+		EpicRtcManager->VideoTrackObserverFactory = MakeRefCount<FEpicRtcVideoTrackObserverFactory>(EpicRtcManager);
+		EpicRtcManager->DataTrackObserverFactory = MakeRefCount<FEpicRtcDataTrackObserverFactory>(EpicRtcManager);
+	}
+
+	FEpicRtcStreamer::~FEpicRtcStreamer()
+	{
+		StopStreaming();
+	}
+
+	void FEpicRtcStreamer::Initialize()
+	{
+		TSharedPtr<FEpicRtcStreamer> Streamer = AsShared();
+
+		InputHandler->GetToStreamerProtocol()->OnProtocolUpdated().AddSP(Streamer.ToSharedRef(), &FEpicRtcStreamer::OnProtocolUpdated);
+		InputHandler->GetFromStreamerProtocol()->OnProtocolUpdated().AddSP(Streamer.ToSharedRef(), &FEpicRtcStreamer::OnProtocolUpdated);
+
+		Streamer->EpicRtcManager->OnSessionStateUpdate.AddSP(Streamer.ToSharedRef(), &FEpicRtcStreamer::OnSessionStateUpdate);
+		Streamer->EpicRtcManager->OnSessionErrorUpdate.AddSP(Streamer.ToSharedRef(), &FEpicRtcStreamer::OnSessionErrorUpdate);
+		Streamer->EpicRtcManager->OnSessionRoomsAvailableUpdate.AddSP(Streamer.ToSharedRef(), &FEpicRtcStreamer::OnSessionRoomsAvailableUpdate);
+
+		Streamer->EpicRtcManager->OnRoomStateUpdate.AddSP(Streamer.ToSharedRef(), &FEpicRtcStreamer::OnRoomStateUpdate);
+		Streamer->EpicRtcManager->OnRoomJoinedUpdate.AddSP(Streamer.ToSharedRef(), &FEpicRtcStreamer::OnRoomJoinedUpdate);
+		Streamer->EpicRtcManager->OnRoomLeftUpdate.AddSP(Streamer.ToSharedRef(), &FEpicRtcStreamer::OnRoomLeftUpdate);
+		Streamer->EpicRtcManager->OnAudioTrackUpdate.AddSP(Streamer.ToSharedRef(), &FEpicRtcStreamer::OnAudioTrackUpdate);
+		Streamer->EpicRtcManager->OnVideoTrackUpdate.AddSP(Streamer.ToSharedRef(), &FEpicRtcStreamer::OnVideoTrackUpdate);
+		Streamer->EpicRtcManager->OnDataTrackUpdate.AddSP(Streamer.ToSharedRef(), &FEpicRtcStreamer::OnDataTrackUpdate);
+		Streamer->EpicRtcManager->OnLocalSdpUpdate.AddSP(Streamer.ToSharedRef(), &FEpicRtcStreamer::OnLocalSdpUpdate);
+		Streamer->EpicRtcManager->OnRemoteSdpUpdate.AddSP(Streamer.ToSharedRef(), &FEpicRtcStreamer::OnRemoteSdpUpdate);
+		Streamer->EpicRtcManager->OnRoomErrorUpdate.AddSP(Streamer.ToSharedRef(), &FEpicRtcStreamer::OnRoomErrorUpdate);
+
+		Streamer->EpicRtcManager->OnAudioTrackMuted.AddSP(Streamer.ToSharedRef(), &FEpicRtcStreamer::OnAudioTrackMuted);
+		Streamer->EpicRtcManager->OnAudioTrackFrame.AddSP(Streamer.ToSharedRef(), &FEpicRtcStreamer::OnAudioTrackFrame);
+		Streamer->EpicRtcManager->OnAudioTrackRemoved.AddSP(Streamer.ToSharedRef(), &FEpicRtcStreamer::OnAudioTrackRemoved);
+		Streamer->EpicRtcManager->OnAudioTrackState.AddSP(Streamer.ToSharedRef(), &FEpicRtcStreamer::OnAudioTrackState);
+
+		Streamer->EpicRtcManager->OnVideoTrackMuted.AddSP(Streamer.ToSharedRef(), &FEpicRtcStreamer::OnVideoTrackMuted);
+		Streamer->EpicRtcManager->OnVideoTrackFrame.AddSP(Streamer.ToSharedRef(), &FEpicRtcStreamer::OnVideoTrackFrame);
+		Streamer->EpicRtcManager->OnVideoTrackRemoved.AddSP(Streamer.ToSharedRef(), &FEpicRtcStreamer::OnVideoTrackRemoved);
+		Streamer->EpicRtcManager->OnVideoTrackState.AddSP(Streamer.ToSharedRef(), &FEpicRtcStreamer::OnVideoTrackState);
+
+		Streamer->EpicRtcManager->OnDataTrackRemoved.AddSP(Streamer.ToSharedRef(), &FEpicRtcStreamer::OnDataTrackRemoved);
+		Streamer->EpicRtcManager->OnDataTrackState.AddSP(Streamer.ToSharedRef(), &FEpicRtcStreamer::OnDataTrackState);
+		Streamer->EpicRtcManager->OnDataTrackMessage.AddSP(Streamer.ToSharedRef(), &FEpicRtcStreamer::OnDataTrackMessage);
+
+		FPixelStreaming2RTCModule::GetModule()->GetStatsCollector()->OnStatsReady.AddSP(Streamer.ToSharedRef(), &FEpicRtcStreamer::OnStatsReady);
+	}
+
+	void FEpicRtcStreamer::OnProtocolUpdated()
+	{
+		Players->Apply([this](FString DataPlayerId, FPlayerContext& PlayerContext) {
+			SendProtocol(DataPlayerId);
+		});
+	}
+
+	void FEpicRtcStreamer::SetStreamFPS(int32 InFramesPerSecond)
+	{
+		VideoSourceGroup->SetFPS(InFramesPerSecond);
+	}
+
+	int32 FEpicRtcStreamer::GetStreamFPS()
+	{
+		return VideoSourceGroup->GetFPS();
+	}
+
+	void FEpicRtcStreamer::SetCoupleFramerate(bool bCouple)
+	{
+		VideoSourceGroup->SetDecoupleFramerate(!bCouple);
+	}
+
+	void FEpicRtcStreamer::SetVideoProducer(TSharedPtr<IPixelStreaming2VideoProducer> Producer)
+	{
+		VideoCapturer->SetVideoProducer(StaticCastSharedPtr<FVideoProducer>(Producer));
+	}
+
+	TWeakPtr<IPixelStreaming2VideoProducer> FEpicRtcStreamer::GetVideoProducer()
+	{
+		return VideoCapturer->GetVideoProducer();
+	}
+
+	void FEpicRtcStreamer::SetConnectionURL(const FString& InConnectionURL)
+	{
+		CurrentSignallingServerURL = InConnectionURL;
+	}
+
+	FString FEpicRtcStreamer::GetConnectionURL()
+	{
+		return CurrentSignallingServerURL;
+	}
+
+	void FEpicRtcStreamer::StartStreaming()
+	{
+		if (CurrentSignallingServerURL.IsEmpty())
+		{
+			UE_LOG(LogPixelStreaming2RTC, Log, TEXT("Attempted to start streamer (%s) but no signalling server URL has been set. Use Streamer->SetConnectionURL(URL) or -PixelStreamingConnectionURL="), *StreamerId);
+			return;
+		}
+
+		StopStreaming();
+		ReconnectTimer->Stop();
+
+		if (UPixelStreaming2Delegates* Delegates = UPixelStreaming2Delegates::Get())
+		{
+			ConsumeStatsHandle = Delegates->OnStatChangedNative.AddSP(AsShared(), &FEpicRtcStreamer::ConsumeStats);
+			AllConnectionsClosedHandle = Delegates->OnAllConnectionsClosedNative.AddSP(AsShared(), &FEpicRtcStreamer::TriggerMouseLeave);
+		}
+
+		VideoCapturer->ResetFrameCapturer();
+
+		// Broadcast the preconnection event just before we do `TryConnect`
+		StreamingPreConnectionEvent.Broadcast(this);
+
+		VideoSourceGroup->Start();
+
+		FUtf8String Utf8StreamerId(StreamerId);
+		FUtf8String Utf8CurrentSignallingServerURL(CurrentSignallingServerURL);
+
+		EpicRtcSessionConfig SessionConfig{
+			._id = ToEpicRtcStringView(Utf8StreamerId),
+			._url = ToEpicRtcStringView(Utf8CurrentSignallingServerURL),
+			._observer = EpicRtcManager->SessionObserver
+		};
+
+		EpicRtcErrorCode Result = EpicRtcManager->EpicRtcConference->CreateSession(SessionConfig, EpicRtcManager->EpicRtcSession.GetInitReference());
+		if (Result != EpicRtcErrorCode::Ok)
+		{
+			UE_LOG(LogPixelStreaming2RTC, Error, TEXT("Failed to create EpicRtc session. CreateSession returned %s"), *ToString(Result));
+			StopStreaming();
+			return;
+		}
+
+		Result = EpicRtcManager->EpicRtcSession->Connect();
+		if (Result != EpicRtcErrorCode::Ok)
+		{
+			UE_LOG(LogPixelStreaming2RTC, Error, TEXT("Failed to connect EpicRtcSession. Connect returned %s"), *ToString(Result));
+			StopStreaming();
+			return;
+		}
+
+		bStreamingStarted = true;
+	}
+
+	void FEpicRtcStreamer::StopStreaming()
+	{
+		if (UPixelStreaming2Delegates* Delegates = UPixelStreaming2Delegates::Get())
+		{
+			Delegates->OnStatChangedNative.Remove(ConsumeStatsHandle);
+			Delegates->OnAllConnectionsClosedNative.Remove(AllConnectionsClosedHandle);
+		}
+
+		RemoveSession(true);
+
+		VideoSourceGroup->Stop();
+		TriggerMouseLeave(StreamerId);
+
+		if (bStreamingStarted)
+		{
+			OnStreamingStopped().Broadcast(this);
+		}
+
+		DeleteAllPlayerSessions();
+		bStreamingStarted = false;
+	}
+
+	void FEpicRtcStreamer::RemoveSession(bool bDisconnect)
+	{
+		if (!EpicRtcManager->EpicRtcSession)
+		{
+			return;
+		}
+
+		RemoveRoom();
+
+		if (bDisconnect)
+		{
+			EpicRtcErrorCode Result = EpicRtcManager->EpicRtcSession->Disconnect(ToEpicRtcStringView("Streaming Session Removed"));
+			if (Result == EpicRtcErrorCode::SessionDisconnected)
+			{
+				UE_LOG(LogPixelStreaming2RTC, VeryVerbose, TEXT("FEpicRtcStreamer::StopStreaming - Session disconnected cleanly."));
+			}
+			else if (Result != EpicRtcErrorCode::Ok)
+			{
+				UE_LOG(LogPixelStreaming2RTC, Error, TEXT("Failed to disconnect EpicRtcSession. Disconnect returned %s"), *ToString(Result));
+			}
+		}
+
+		FUtf8String Utf8StreamerId(StreamerId);
+		EpicRtcManager->EpicRtcConference->RemoveSession(ToEpicRtcStringView(Utf8StreamerId));
+
+		EpicRtcManager->EpicRtcSession = nullptr;
+	}
+
+	void FEpicRtcStreamer::RemoveRoom()
+	{
+		if (!EpicRtcManager->EpicRtcRoom)
+		{
+			return;
+		}
+
+		FUtf8String Utf8StreamerId(StreamerId);
+		EpicRtcManager->EpicRtcRoom->Leave();
+		EpicRtcManager->EpicRtcSession->RemoveRoom(ToEpicRtcStringView(Utf8StreamerId));
+
+		EpicRtcManager->EpicRtcRoom = nullptr;
+	}
+
+	void FEpicRtcStreamer::OnStatsReady(const FString& PlayerId, const EpicRtcConnectionStats& ConnectionStats)
+	{
+		FPlayerContext* PlayerContext = Players->Find(PlayerId);
+		if (!PlayerContext)
+		{
+			return;
+		}
+
+		if (!PlayerContext->StatsCollector)
+		{
+			return;
+		}
+
+		PlayerContext->StatsCollector->Process(ConnectionStats);
+	}
+
+	IPixelStreaming2Streamer::FPreConnectionEvent& FEpicRtcStreamer::OnPreConnection()
+	{
+		return StreamingPreConnectionEvent;
+	}
+
+	IPixelStreaming2Streamer::FStreamingStartedEvent& FEpicRtcStreamer::OnStreamingStarted()
+	{
+		return StreamingStartedEvent;
+	}
+
+	IPixelStreaming2Streamer::FStreamingStoppedEvent& FEpicRtcStreamer::OnStreamingStopped()
+	{
+		return StreamingStoppedEvent;
+	}
+
+	void FEpicRtcStreamer::ForceKeyFrame()
+	{
+		VideoSourceGroup->ForceKeyFrame();
+	}
+
+	void FEpicRtcStreamer::FreezeStream(UTexture2D* Texture)
+	{
+		FreezeFrame->StartFreeze(Texture);
+	}
+
+	void FEpicRtcStreamer::UnfreezeStream()
+	{
+		// Force a keyframe so when stream unfreezes if player has never received a frame before they can still connect.
+		ForceKeyFrame();
+		FreezeFrame->StopFreeze();
+	}
+
+	void FEpicRtcStreamer::SendAllPlayersMessage(FString MessageType, const FString& Descriptor)
+	{
+		Players->Apply([&MessageType, &Descriptor](FString PlayerId, FPlayerContext& PlayerContext) {
+			if (PlayerContext.DataTrack && !IsSFU(PlayerId))
+			{
+				PlayerContext.DataTrack->SendMessage(MessageType, Descriptor);
+			}
+		});
+	}
+
+	void FEpicRtcStreamer::SendPlayerMessage(FString PlayerId, FString MessageType, const FString& Descriptor)
+	{
+		if (IsSFU(PlayerId))
+		{
+			return;
+		}
+		if (FPlayerContext* PlayerContext = Players->Find(PlayerId))
+		{
+			if (!PlayerContext->DataTrack)
+			{
+				return;
+			}
+			PlayerContext->DataTrack->SendMessage(MessageType, Descriptor);
+		}
+	}
+
+	void FEpicRtcStreamer::SendFileData(const TArray64<uint8>& ByteData, FString& MimeType, FString& FileExtension)
+	{
+		// TODO this should be dispatched as an async task, but because we lock when we visit the data
+		// channels it might be a bad idea. At some point it would be good to take a snapshot of the
+		// keys in the map when we start, then one by one get the channel and send the data
+
+		Players->Apply([&ByteData, &MimeType, &FileExtension](FString PlayerId, FPlayerContext& PlayerContext) {
+			if (!PlayerContext.DataTrack)
+			{
+				return;
+			}
+
+			// Send the mime type first
+			PlayerContext.DataTrack->SendMessage(EPixelStreaming2FromStreamerMessage::FileMimeType, MimeType);
+
+			// Send the extension next
+			PlayerContext.DataTrack->SendMessage(EPixelStreaming2FromStreamerMessage::FileExtension, FileExtension);
+
+			// Send the contents of the file. Note to callers: consider running this on its own thread, it can take a while if the file is big.
+			PlayerContext.DataTrack->SendArbitraryData(EPixelStreaming2FromStreamerMessage::FileContents, ByteData);
+		});
+	}
+
+	void FEpicRtcStreamer::KickPlayer(FString PlayerId)
+	{
+		if (FPlayerContext* PlayerContext = Players->Find(PlayerId))
+		{
+			PlayerContext->ParticipantInterface->Kick();
+		}
+	}
+
+	TArray<FString> FEpicRtcStreamer::GetConnectedPlayers()
+	{
+		TArray<FString> ConnectedPlayerIds;
+		Players->Apply([&ConnectedPlayerIds, this](FString PlayerId, FPlayerContext& PlayerContext) {
+			ConnectedPlayerIds.Add(PlayerId);
+		});
+		return ConnectedPlayerIds;
+	}
+
+	IPixelStreaming2AudioSink* FEpicRtcStreamer::GetPeerAudioSink(FString PlayerId)
+	{
+		if (FPlayerContext* PlayerContext = Players->Find(PlayerId))
+		{
+			if (PlayerContext->AudioSink)
+			{
+				return PlayerContext->AudioSink.Get();
+			}
+		}
+		return nullptr;
+	}
+
+	IPixelStreaming2AudioSink* FEpicRtcStreamer::GetUnlistenedAudioSink()
+	{
+		IPixelStreaming2AudioSink* Result = nullptr;
+		Players->ApplyUntil([&Result](FString PlayerId, FPlayerContext& PlayerContext) {
+			if (PlayerContext.AudioSink)
+			{
+				if (!PlayerContext.AudioSink->HasAudioConsumers())
+				{
+					Result = PlayerContext.AudioSink.Get();
+					return true;
+				}
+			}
+			return false;
+		});
+		return Result;
+	}
+
+	IPixelStreaming2VideoSink* FEpicRtcStreamer::GetPeerVideoSink(FString PlayerId)
+	{
+		if (FPlayerContext* PlayerContext = Players->Find(PlayerId))
+		{
+			if (PlayerContext->VideoSink)
+			{
+				return PlayerContext->VideoSink.Get();
+			}
+		}
+		return nullptr;
+	}
+
+	IPixelStreaming2VideoSink* FEpicRtcStreamer::GetUnwatchedVideoSink()
+	{
+		IPixelStreaming2VideoSink* Result = nullptr;
+		Players->ApplyUntil([&Result](FString PlayerId, FPlayerContext& PlayerContext) {
+			if (PlayerContext.VideoSink)
+			{
+				if (!PlayerContext.VideoSink->HasVideoConsumers())
+				{
+					Result = PlayerContext.VideoSink.Get();
+					return true;
+				}
+			}
+			return false;
+		});
+		return Result;
+	}
+
+	void FEpicRtcStreamer::SetConfigOption(const FName& OptionName, const FString& Value)
+	{
+		if (Value.IsEmpty())
+		{
+			ConfigOptions.Remove(OptionName);
+		}
+		else
+		{
+			ConfigOptions.Add(OptionName, Value);
+		}
+	}
+
+	bool FEpicRtcStreamer::GetConfigOption(const FName& OptionName, FString& OutValue)
+	{
+		FString* OptionValue = ConfigOptions.Find(OptionName);
+		if (OptionValue)
+		{
+			OutValue = *OptionValue;
+			return true;
+		}
+		else
+		{
+			return false;
+		}
+	}
+
+	void FEpicRtcStreamer::PlayerRequestsBitrate(FString PlayerId, int MinBitrate, int MaxBitrate)
+	{
+		UPixelStreaming2PluginSettings::CVarWebRTCMinBitrate.AsVariable()->Set(MinBitrate);
+		UPixelStreaming2PluginSettings::CVarWebRTCMaxBitrate.AsVariable()->Set(MaxBitrate);
+	}
+
+	void FEpicRtcStreamer::RefreshStreamBitrate()
+	{
+		Players->Apply([this](FString PlayerId, FPlayerContext& PlayerContext) {
+			if (!PlayerContext.ParticipantInterface)
+			{
+				return;
+			}
+
+			TRefCountPtr<EpicRtcConnectionInterface> ConnectionInterface = PlayerContext.ParticipantInterface->GetConnection();
+			if (!ConnectionInterface)
+			{
+				return;
+			}
+
+			EpicRtcBitrate Bitrates = {
+				._minBitrateBps = UPixelStreaming2PluginSettings::CVarWebRTCMinBitrate.GetValueOnAnyThread(),
+				._hasMinBitrateBps = true,
+				._maxBitrateBps = UPixelStreaming2PluginSettings::CVarWebRTCMaxBitrate.GetValueOnAnyThread(),
+				._hasMaxBitrateBps = true,
+				._startBitrateBps = UPixelStreaming2PluginSettings::CVarWebRTCStartBitrate.GetValueOnAnyThread(),
+				._hasStartBitrateBps = true
+			};
+
+			ConnectionInterface->SetConnectionRates(Bitrates);
+		});
+	}
+
+	void FEpicRtcStreamer::ForEachPlayer(const TFunction<void(FString, FPlayerContext)>& Func)
+	{
+		Players->Apply(Func);
+	}
+
+	void FEpicRtcStreamer::ConsumeStats(FString PlayerId, FName StatName, float StatValue)
+	{
+		if (IsSFU(PlayerId))
+		{
+			return;
+		}
+
+		if (StatName != PixelStreaming2StatNames::MeanQPPerSecond)
+		{
+			return;
+		}
+
+		FPlayerContext* PlayerContext = Players->Find(PlayerId);
+
+		if (!PlayerContext)
+		{
+			return;
+		}
+
+		if (!PlayerContext->DataTrack)
+		{
+			return;
+		}
+
+		PlayerContext->DataTrack->SendMessage(EPixelStreaming2FromStreamerMessage::VideoEncoderAvgQP, FString::FromInt((int)StatValue));
+	}
+
+	void FEpicRtcStreamer::DeletePlayerSession(FString PlayerId)
+	{
+		// We dont want to allow this to be deleted within Players.Remove because
+		// we lock the players map and the delete could dispatch a webrtc object
+		// delete on the signalling thread which might be waiting for the players
+		// lock.
+		FPlayerContext PendingDeletePlayer;
+		if (FPlayerContext* PlayerContext = Players->Find(PlayerId))
+		{
+			// when a sfu is connected we only get disconnect messages.
+			// we dont get connect messages but we might get datachannel requests which can result
+			// in players with no PeerConnection but a datachannel
+			if (PlayerContext->VideoSource)
+			{
+				VideoSourceGroup->RemoveVideoSource(PlayerContext->VideoSource.Get());
+			}
+
+			// Close any data track related things (we do this here because RoomLeft happens before DataTrack stopped fires)
+			// So if we only did this in DataTrack stooped the playerId would already be removed
+			OnDataChannelClosed(PlayerId);
+
+			PendingDeletePlayer = *PlayerContext;
+		}
+
+		Players->Remove(PlayerId);
+
+		if (UPixelStreaming2Delegates* Delegates = UPixelStreaming2Delegates::Get())
+		{
+			Delegates->OnClosedConnection.Broadcast(StreamerId, PlayerId);
+			Delegates->OnClosedConnectionNative.Broadcast(StreamerId, PlayerId);
+			if (Players->IsEmpty())
+			{
+				Delegates->OnAllConnectionsClosed.Broadcast(StreamerId);
+				Delegates->OnAllConnectionsClosedNative.Broadcast(StreamerId);
+			}
+		}
+
+		if (FStats* PSStats = FStats::Get())
+		{
+			PSStats->RemovePeerStats(PlayerId);
+		}
+	}
+
+	void FEpicRtcStreamer::DeleteAllPlayerSessions()
+	{
+		if (FStats* PSStats = FStats::Get())
+		{
+			PSStats->RemoveAllPeerStats();
+		}
+
+		VideoSourceGroup->RemoveAllVideoSources();
+		Players->Empty();
+		InputControllingId = INVALID_PLAYER_ID;
+		if (UPixelStreaming2Delegates* Delegates = UPixelStreaming2Delegates::Get())
+		{
+			Delegates->OnAllConnectionsClosed.Broadcast(StreamerId);
+			Delegates->OnAllConnectionsClosedNative.Broadcast(StreamerId);
+		}
+	}
+
+	void FEpicRtcStreamer::OnDataChannelOpen(FString PlayerId)
+	{
+		// Only time we automatically make a new peer the input controlling host is if they are the first peer (and not the SFU).
+		bool HostControlsInput = GetEnumFromCVar<EInputControllerMode>(UPixelStreaming2PluginSettings::CVarInputController) == EInputControllerMode::Host;
+		if (HostControlsInput && !IsSFU(PlayerId) && InputControllingId == INVALID_PLAYER_ID)
+		{
+			InputControllingId = PlayerId;
+		}
+
+		if (UPixelStreaming2Delegates* Delegates = UPixelStreaming2Delegates::Get())
+		{
+			Delegates->OnDataTrackOpen.Broadcast(StreamerId, PlayerId);
+			Delegates->OnDataTrackOpenNative.Broadcast(StreamerId, PlayerId);
+		}
+
+		// When data channel is open
+		SendProtocol(PlayerId);
+		// Try to send cached freeze frame (if we have one)
+		FreezeFrame->SendCachedFreezeFrameTo(PlayerId);
+		SendInitialSettings(PlayerId);
+		SendPeerControllerMessages(PlayerId);
+	}
+
+	void FEpicRtcStreamer::OnDataChannelClosed(FString PlayerId)
+	{
+		if (FPlayerContext* PlayerContext = Players->Find(PlayerId))
+		{
+			PlayerContext->DataTrack = nullptr;
+
+			if (InputControllingId == PlayerId)
+			{
+				InputControllingId = INVALID_PLAYER_ID;
+				// just get the first channel we have and give it input control.
+				Players->ApplyUntil([this](FString PlayerId, FPlayerContext& PlayerContext) {
+					if (!PlayerContext.DataTrack)
+					{
+						return false;
+					}
+					if (IsSFU(PlayerId))
+					{
+						return false;
+					}
+					InputControllingId = PlayerId;
+					PlayerContext.DataTrack->SendMessage(EPixelStreaming2FromStreamerMessage::InputControlOwnership, 1 /* ControlsInput */);
+					return true;
+				});
+			}
+
+			if (UPixelStreaming2Delegates* Delegates = UPixelStreaming2Delegates::Get())
+			{
+				Delegates->OnDataTrackClosed.Broadcast(StreamerId, PlayerId);
+				Delegates->OnDataTrackClosedNative.Broadcast(StreamerId, PlayerId);
+			}
+		}
+	}
+
+	void FEpicRtcStreamer::SendInitialSettings(FString PlayerId) const
+	{
+		const FString PixelStreaming2Payload = FString::Printf(TEXT("{ \"AllowPixelStreamingCommands\": %s, \"DisableLatencyTest\": %s }"),
+			UPixelStreaming2PluginSettings::CVarInputAllowConsoleCommands.GetValueOnAnyThread() ? TEXT("true") : TEXT("false"),
+			UPixelStreaming2PluginSettings::CVarDisableLatencyTester.GetValueOnAnyThread() ? TEXT("true") : TEXT("false"));
+
+		const FString WebRTCPayload = FString::Printf(TEXT("{ \"FPS\": %d, \"MinBitrate\": %d, \"MaxBitrate\": %d }"),
+			UPixelStreaming2PluginSettings::CVarWebRTCFps.GetValueOnAnyThread(),
+			UPixelStreaming2PluginSettings::CVarWebRTCMinBitrate.GetValueOnAnyThread(),
+			UPixelStreaming2PluginSettings::CVarWebRTCMaxBitrate.GetValueOnAnyThread());
+
+		const FString EncoderPayload = FString::Printf(TEXT("{ \"TargetBitrate\": %d, \"MinQuality\": %d, \"MaxQuality\": %d }"),
+			UPixelStreaming2PluginSettings::CVarEncoderTargetBitrate.GetValueOnAnyThread(),
+			UPixelStreaming2PluginSettings::CVarEncoderMinQuality.GetValueOnAnyThread(),
+			UPixelStreaming2PluginSettings::CVarEncoderMaxQuality.GetValueOnAnyThread());
+
+		FString ConfigPayload = TEXT("{ ");
+		bool	bComma = false; // Simplest way to avoid complaints from pedantic JSON parsers
+		for (const TPair<FName, FString>& Option : ConfigOptions)
+		{
+			if (bComma)
+			{
+				ConfigPayload.Append(TEXT(", "));
+			}
+			ConfigPayload.Append(FString::Printf(TEXT("\"%s\": \"%s\""), *Option.Key.ToString(), *Option.Value));
+			bComma = true;
+		}
+		ConfigPayload.Append(TEXT("}"));
+
+		const FString FullPayload = FString::Printf(TEXT("{ \"PixelStreaming\": %s, \"Encoder\": %s, \"WebRTC\": %s, \"ConfigOptions\": %s }"), *PixelStreaming2Payload, *EncoderPayload, *WebRTCPayload, *ConfigPayload);
+
+		if (const FPlayerContext* PlayerContext = Players->Find(PlayerId))
+		{
+			if (!PlayerContext->DataTrack)
+			{
+				return;
+			}
+			PlayerContext->DataTrack->SendMessage(EPixelStreaming2FromStreamerMessage::InitialSettings, FullPayload);
+		}
+	}
+
+	void FEpicRtcStreamer::SendProtocol(FString PlayerId) const
+	{
+		const TArray<TSharedPtr<IPixelStreaming2DataProtocol>> Protocols = { InputHandler->GetToStreamerProtocol(), InputHandler->GetFromStreamerProtocol() };
+		for (TSharedPtr<IPixelStreaming2DataProtocol> Protocol : Protocols)
+		{
+			TSharedPtr<FJsonObject>	  ProtocolJson = Protocol->ToJson();
+			FString					  Body;
+			TSharedRef<TJsonWriter<>> JsonWriter = TJsonWriterFactory<>::Create(&Body);
+			if (!ensure(FJsonSerializer::Serialize(ProtocolJson.ToSharedRef(), JsonWriter)))
+			{
+				UE_LOG(LogPixelStreaming2RTC, Warning, TEXT("Cannot serialize protocol json object"));
+				return;
+			}
+
+			if (const FPlayerContext* PlayerContext = Players->Find(PlayerId))
+			{
+				if (!PlayerContext->DataTrack)
+				{
+					return;
+				}
+				PlayerContext->DataTrack->SendMessage(EPixelStreaming2FromStreamerMessage::Protocol, Body);
+			}
+		}
+	}
+
+	void FEpicRtcStreamer::SendPeerControllerMessages(FString PlayerId) const
+	{
+		if (const FPlayerContext* PlayerContext = Players->Find(PlayerId))
+		{
+			if (!PlayerContext->DataTrack)
+			{
+				return;
+			}
+			const uint8 ControlsInput = (GetEnumFromCVar<EInputControllerMode>(UPixelStreaming2PluginSettings::CVarInputController) == EInputControllerMode::Host) ? (PlayerId == InputControllingId) : 1;
+			// Even though the QualityController feature is removed we send it for backwards compatibility with older frontends (can probably remove 2 versions after 5.5)
+			PlayerContext->DataTrack->SendMessage(EPixelStreaming2FromStreamerMessage::InputControlOwnership, ControlsInput);
+			PlayerContext->DataTrack->SendMessage(EPixelStreaming2FromStreamerMessage::QualityControlOwnership, 1 /* True */);
+		}
+	}
+
+	void FEpicRtcStreamer::SendLatencyReport(FString PlayerId) const
+	{
+		if (UPixelStreaming2PluginSettings::CVarDisableLatencyTester.GetValueOnAnyThread())
+		{
+			return;
+		}
+
+		double ReceiptTimeMs = FPlatformTime::ToMilliseconds64(FPlatformTime::Cycles64());
+
+		AsyncTask(ENamedThreads::GameThread, [this, PlayerId, ReceiptTimeMs]() {
+			FString ReportToTransmitJSON;
+
+			if (!UPixelStreaming2PluginSettings::CVarWebRTCDisableStats.GetValueOnAnyThread())
+			{
+				double EncodeMs = -1.0;
+				double CaptureToSendMs = 0.0;
+
+				FStats* Stats = FStats::Get();
+				if (Stats)
+				{
+					Stats->QueryPeerStat(PlayerId, RTCStatCategories::LocalVideoTrack, PixelStreaming2StatNames::MeanEncodeTime, EncodeMs);
+					Stats->QueryPeerStat(PlayerId, RTCStatCategories::LocalVideoTrack, PixelStreaming2StatNames::AvgSendDelay, CaptureToSendMs);
+				}
+
+				double TransmissionTimeMs = FPlatformTime::ToMilliseconds64(FPlatformTime::Cycles64());
+				ReportToTransmitJSON = FString::Printf(
+					TEXT("{ \"ReceiptTimeMs\": %.2f, \"EncodeMs\": %.2f, \"CaptureToSendMs\": %.2f, \"TransmissionTimeMs\": %.2f }"),
+					ReceiptTimeMs,
+					EncodeMs,
+					CaptureToSendMs,
+					TransmissionTimeMs);
+			}
+			else
+			{
+				double TransmissionTimeMs = FPlatformTime::ToMilliseconds64(FPlatformTime::Cycles64());
+				ReportToTransmitJSON = FString::Printf(
+					TEXT("{ \"ReceiptTimeMs\": %.2f, \"EncodeMs\": \"Pixel Streaming stats are disabled\", \"CaptureToSendMs\": \"Pixel Streaming stats are disabled\", \"TransmissionTimeMs\": %.2f }"),
+					ReceiptTimeMs,
+					TransmissionTimeMs);
+			}
+
+			if (const FPlayerContext* PlayerContext = Players->Find(PlayerId))
+			{
+				if (PlayerContext->DataTrack)
+				{
+					PlayerContext->DataTrack->SendMessage(EPixelStreaming2FromStreamerMessage::LatencyTest, ReportToTransmitJSON);
+				}
+			}
+		});
+	}
+
+	void FEpicRtcStreamer::HandleRelayStatusMessage(const uint8_t* Data, uint32_t Size, EpicRtcDataTrackInterface* DataTrack)
+	{
+		// skip type
+		Data++;
+		Size--;
+		FString														   PlayerId = ReadString(Data, Size);
+		checkf(Size > 0, TEXT("Malformed relay status message!")) bool bIsOn = static_cast<bool>(Data[0]);
+
+		FString DataTrackId = ToString(DataTrack->GetId());
+		if (bIsOn)
+		{
+			UE_LOG(LogPixelStreaming2RTC, Log, TEXT("FEpicRtcStreamer::HandleRelayStatusMessage() Adding new PlayerId [%s] with DataTrackId [%s]"), *PlayerId, *DataTrackId);
+
+			FString SFUId;
+			if (FindPlayerByDataTrack(DataTrack, SFUId))
+			{
+				FPlayerContext* SFUContext = Players->Find(SFUId);
+				FPlayerContext& PlayerContext = Players->FindOrAdd(PlayerId);
+				PlayerContext.DataTrack = FEpicRtcMutliplexDataTrack::Create(SFUContext->DataTrack, InputHandler->GetFromStreamerProtocol(), PlayerId);
+				OnDataChannelOpen(PlayerId);
+			}
+			else
+			{
+				UE_LOG(LogPixelStreaming2RTC, Log, TEXT("FEpicRtcStreamer::HandleRelayStatusMessage() Failed to find SFU PlayerContext"));
+			}
+		}
+		else
+		{
+			UE_LOG(LogPixelStreaming2RTC, Log, TEXT("FEpicRtcStreamer::HandleRelayStatusMessage() Removing PlayerId [%s] with DataTrackId [%s]"), *PlayerId, *DataTrackId);
+
+			OnDataChannelClosed(PlayerId);
+			Players->Remove(PlayerId);
+		}
+	}
+
+	void FEpicRtcStreamer::TriggerMouseLeave(FString InStreamerId)
+	{
+		if (!IsEngineExitRequested() && StreamerId == InStreamerId)
+		{
+			TSharedPtr<IPixelStreaming2InputHandler> SharedInputHandler = InputHandler;
+
+			// Force a MouseLeave event. This prevents the PixelStreaming2ApplicationWrapper from
+			// still wrapping the base FSlateApplication after we stop streaming
+			const auto MouseLeaveFunction = [SharedInputHandler]() {
+				if (SharedInputHandler.IsValid())
+				{
+					TArray<uint8>							EmptyArray;
+					TFunction<void(FString, FMemoryReader)> MouseLeaveHandler = SharedInputHandler->FindMessageHandler("MouseLeave");
+					MouseLeaveHandler("", FMemoryReader(EmptyArray));
+				}
+			};
+
+			if (IsInGameThread())
+			{
+				MouseLeaveFunction();
+			}
+			else
+			{
+				AsyncTask(ENamedThreads::GameThread, [MouseLeaveFunction]() {
+					MouseLeaveFunction();
+				});
+			}
+		}
+	}
+
+	void FEpicRtcStreamer::OnSessionStateUpdate(const EpicRtcSessionState State)
+	{
+		switch (State)
+		{
+			case EpicRtcSessionState::Connected:
+			{
+				bSignallingConnected = true;
+				if (UPixelStreaming2Delegates* Delegates = UPixelStreaming2Delegates::Get())
+				{
+					Delegates->OnConnectedToSignallingServer.Broadcast(StreamerId);
+					Delegates->OnConnectedToSignallingServerNative.Broadcast(StreamerId);
+				}
+
+				UE_LOG(LogPixelStreaming2RTC, VeryVerbose, TEXT("FEpicRtcStreamer::OnSessionStateUpdate State=Connected"));
+				EpicRtcBitrate Bitrate = {
+					._minBitrateBps = UPixelStreaming2PluginSettings::CVarWebRTCMinBitrate.GetValueOnAnyThread(),
+					._hasMinBitrateBps = true,
+					._maxBitrateBps = UPixelStreaming2PluginSettings::CVarWebRTCMaxBitrate.GetValueOnAnyThread(),
+					._hasMaxBitrateBps = true,
+					._startBitrateBps = UPixelStreaming2PluginSettings::CVarWebRTCStartBitrate.GetValueOnAnyThread(),
+					._hasStartBitrateBps = true
+				};
+
+				EpicRtcPortAllocator PortAllocator = {
+					._minPort = UPixelStreaming2PluginSettings::CVarWebRTCMinPort.GetValueOnAnyThread(),
+					._hasMinPort = true,
+					._maxPort = UPixelStreaming2PluginSettings::CVarWebRTCMaxPort.GetValueOnAnyThread(),
+					._hasMaxPort = true,
+					._portAllocation = static_cast<EpicRtcPortAllocatorOptions>(UPixelStreaming2PluginSettings::GetPortAllocationFlags())
+				};
+
+				EpicRtcConnectionConfig ConnectionConfig = {
+					._iceServers = { ._ptr = nullptr, ._size = 0 }, // This can stay empty because EpicRtc handles the ice servers internally
+					._portAllocator = PortAllocator,
+					._bitrate = Bitrate,
+					._iceConnectionPolicy = EpicRtcIcePolicy::All,
+					._disableTcpCandidates = false
+				};
+
+				FUtf8String		  Utf8StreamerId(StreamerId);
+				EpicRtcRoomConfig RoomConfig = {
+					._id = ToEpicRtcStringView(Utf8StreamerId),
+					._connectionConfig = ConnectionConfig,
+					._ticket = EpicRtcStringView{ ._ptr = nullptr, ._length = 0 },
+					._observer = EpicRtcManager->RoomObserver,
+					._audioTrackObserverFactory = EpicRtcManager->AudioTrackObserverFactory,
+					._dataTrackObserverFactory = EpicRtcManager->DataTrackObserverFactory,
+					._videoTrackObserverFactory = EpicRtcManager->VideoTrackObserverFactory
+				};
+
+				EpicRtcErrorCode Result = EpicRtcManager->EpicRtcSession->CreateRoom(RoomConfig, EpicRtcManager->EpicRtcRoom.GetInitReference());
+				if (Result != EpicRtcErrorCode::Ok)
+				{
+					UE_LOG(LogPixelStreaming2RTC, Error, TEXT("Failed to create EpicRtc room. CreateRoom returned %s"), *ToString(Result));
+					break;
+				}
+
+				EpicRtcManager->EpicRtcRoom->Join();
+
+				// Would be better renamed to OnSessionConnected
+				OnStreamingStarted().Broadcast(this);
+			}
+			break;
+			case EpicRtcSessionState::New:
+				// Do something on `new` Session here
+				UE_LOG(LogPixelStreaming2RTC, VeryVerbose, TEXT("FEpicRtcStreamer::OnSessionStateUpdate State=New"));
+				break;
+			case EpicRtcSessionState::Pending:
+				// Do something on `pending` session here
+				UE_LOG(LogPixelStreaming2RTC, VeryVerbose, TEXT("FEpicRtcStreamer::OnSessionStateUpdate State=Pending"));
+				break;
+			case EpicRtcSessionState::Disconnected:
+				// Do something on `disconnected` session here
+				bSignallingConnected = false;
+				if (UPixelStreaming2Delegates* Delegates = UPixelStreaming2Delegates::Get())
+				{
+					Delegates->OnDisconnectedFromSignallingServer.Broadcast(StreamerId);
+					Delegates->OnDisconnectedFromSignallingServerNative.Broadcast(StreamerId);
+				}
+				UE_LOG(LogPixelStreaming2RTC, VeryVerbose, TEXT("FEpicRtcStreamer::OnSessionStateUpdate State=Disconnected"));
+				RemoveSession(false);
+				StopStreaming();
+				ReconnectTimer->Start(AsShared());
+				break;
+			case EpicRtcSessionState::Failed:
+				// Do something on `failed` session here
+				bSignallingConnected = false;
+				UE_LOG(LogPixelStreaming2RTC, VeryVerbose, TEXT("FEpicRtcStreamer::OnSessionStateUpdate State=Failed"));
+				break;
+			case EpicRtcSessionState::Exiting:
+				// Do something on `exiting` session here
+				UE_LOG(LogPixelStreaming2RTC, VeryVerbose, TEXT("FEpicRtcStreamer::OnSessionStateUpdate State=Exiting"));
+				break;
+			default:
+				UE_LOG(LogPixelStreaming2RTC, Error, TEXT("FEpicRtcStreamer::OnSessionStateUpdate An unhandled session state was encountered. This switch might be missing a case."));
+				break;
+		}
+	}
+
+	void FEpicRtcStreamer::OnSessionErrorUpdate(const EpicRtcErrorCode ErrorUpdate)
+	{
+		UE_LOGFMT(LogPixelStreaming2RTC, VeryVerbose, "FEpicRtcStreamer::OnSessionErrorUpdate does nothing");
+	}
+
+	void FEpicRtcStreamer::OnSessionRoomsAvailableUpdate(EpicRtcStringArrayInterface* RoomsList)
+	{
+		UE_LOGFMT(LogPixelStreaming2RTC, VeryVerbose, "FEpicRtcStreamer::OnSessionRoomsAvailableUpdate does nothing");
+	}
+
+	void FEpicRtcStreamer::OnRoomStateUpdate(const EpicRtcRoomState State)
+	{
+		UE_LOGFMT(LogPixelStreaming2RTC, VeryVerbose, "FEpicRtcStreamer::OnRoomStateUpdate does nothing");
+	}
+
+	FUtf8String GetAudioStreamID()
+	{
+		const bool bSyncVideoAndAudio = !UPixelStreaming2PluginSettings::CVarWebRTCDisableAudioSync.GetValueOnAnyThread();
+		return bSyncVideoAndAudio ? "pixelstreaming_av_stream_id" : "pixelstreaming_audio_stream_id";
+	}
+
+	FUtf8String GetVideoStreamID()
+	{
+		const bool bSyncVideoAndAudio = !UPixelStreaming2PluginSettings::CVarWebRTCDisableAudioSync.GetValueOnAnyThread();
+		return bSyncVideoAndAudio ? "pixelstreaming_av_stream_id" : "pixelstreaming_video_stream_id";
+	}
+
+	void FEpicRtcStreamer::OnRoomJoinedUpdate(EpicRtcParticipantInterface* Participant)
+	{
+		FString ParticipantId = ToString(Participant->GetId());
+		UE_LOG(LogPixelStreaming2RTC, Log, TEXT("Player (%s) joined"), *ParticipantId);
+
+		if (ParticipantId == StreamerId)
+		{
+			return;
+		}
+
+		if (UPixelStreaming2Delegates* Delegates = UPixelStreaming2Delegates::Get())
+		{
+			Delegates->OnNewConnection.Broadcast(StreamerId, ParticipantId);
+			Delegates->OnNewConnectionNative.Broadcast(StreamerId, ParticipantId);
+		}
+
+		FPlayerContext& PlayerContext = Players->FindOrAdd(ParticipantId);
+		PlayerContext.ParticipantInterface = Participant;
+		PlayerContext.StatsCollector = FRTCStatsCollector::Create(ParticipantId);
+
+		TRefCountPtr<EpicRtcConnectionInterface> ParticipantConnection = Participant->GetConnection();
+		ParticipantConnection->SetManualNegotiation(true);
+
+		const EVideoCodec SelectedCodec = GetEnumFromCVar<EVideoCodec>(UPixelStreaming2PluginSettings::CVarEncoderCodec);
+		const bool		  bNegotiateCodecs = UPixelStreaming2PluginSettings::CVarWebRTCNegotiateCodecs.GetValueOnAnyThread();
+		const bool		  bTransmitUEVideo = !UPixelStreaming2PluginSettings::CVarWebRTCDisableTransmitVideo.GetValueOnAnyThread();
+		bool			  bReceiveBrowserVideo = !UPixelStreaming2PluginSettings::CVarWebRTCDisableReceiveVideo.GetValueOnAnyThread();
+
+		// Check if the user has selected only H.264 on a AMD gpu and disable receiving video.
+		// WebRTC does not support using SendRecv if the encoding and decoding do not support the same codec.
+		// AMD GPUs currently have decoding disabled so WebRTC fails to create SDP codecs with SendRecv.
+		if (IsRHIDeviceAMD() && !bNegotiateCodecs && SelectedCodec == EVideoCodec::H264)
+		{
+			if (bReceiveBrowserVideo)
+			{
+				bReceiveBrowserVideo = false;
+				UE_LOGFMT(LogPixelStreaming2RTC, Warning, "AMD GPUs do not support receiving H.264 video.");
+			}
+		}
+
+		if (bTransmitUEVideo || bReceiveBrowserVideo)
+		{
+			TArray<EpicRtcVideoEncodingConfig> VideoEncodingConfigs;
+			// We need ensure the Rids have the same lifetime as the VideoEncodingConfigs
+			// to ensure the contents don't get deleted before we can call AddVideoSource
+			TArray<FUtf8String> Rids;
+
+			int MaxFramerate = UPixelStreaming2PluginSettings::CVarWebRTCFps.GetValueOnAnyThread();
+
+			TArray<FPixelStreaming2SimulcastLayer> SimulcastParams = UE::PixelStreaming2::GetSimulcastParameters();
+			for (int i = 0; i < SimulcastParams.Num(); ++i)
+			{
+				const FPixelStreaming2SimulcastLayer& SpatialLayer = SimulcastParams[i];
+
+				if (SimulcastParams.Num() > 1)
+				{
+					Rids.Add(FUtf8String("simulcast") + FUtf8String::FromInt(SimulcastParams.Num() - i));
+				}
+
+				EpicRtcVideoEncodingConfig VideoEncodingConfig = {
+					// clang-format off
+					// TODO (Migration): RTCP-7027 Maybe bug in EpicRtc? Setting an rid if there's only one config results in no video
+					._rid = SimulcastParams.Num() > 1
+						? EpicRtcStringView{._ptr = (const char*)(*(Rids[i])), ._length = static_cast<uint64>(Rids[i].Len()) }
+						: EpicRtcStringView{._ptr = nullptr, ._length = 0 },
+					// clang-format on
+					._scaleResolutionDownBy = SpatialLayer.Scaling,
+					._scalabilityMode = static_cast<EpicRtcVideoScalabilityMode>(GetEnumFromCVar<EScalabilityMode>(UPixelStreaming2PluginSettings::CVarEncoderScalabilityMode)), // HACK if the Enums become un-aligned
+					._minBitrate = (uint32_t)SpatialLayer.MinBitrate,
+					._maxBitrate = (uint32_t)SpatialLayer.MaxBitrate,
+					._maxFrameRate = (uint8_t)MaxFramerate
+				};
+
+				VideoEncodingConfigs.Add(VideoEncodingConfig);
+			}
+
+			EpicRtcVideoEncodingConfigSpan VideoEncodingConfigSpan = {
+				._ptr = VideoEncodingConfigs.GetData(),
+				._size = (uint64_t)VideoEncodingConfigs.Num()
+			};
+
+			EpicRtcMediaSourceDirection VideoDirection;
+			if (bTransmitUEVideo && bReceiveBrowserVideo)
+			{
+				VideoDirection = EpicRtcMediaSourceDirection::SendRecv;
+			}
+			else if (bTransmitUEVideo)
+			{
+				VideoDirection = EpicRtcMediaSourceDirection::SendOnly;
+			}
+			else if (bReceiveBrowserVideo)
+			{
+				VideoDirection = EpicRtcMediaSourceDirection::RecvOnly;
+			}
+			else
+			{
+				VideoDirection = EpicRtcMediaSourceDirection::RecvOnly;
+			}
+
+			FUtf8String		   VideoStreamID = GetVideoStreamID();
+			EpicRtcVideoSource VideoSource = {
+				._streamId = ToEpicRtcStringView(VideoStreamID),
+				._encodings = VideoEncodingConfigSpan,
+				._direction = VideoDirection
+			};
+
+			ParticipantConnection->AddVideoSource(VideoSource);
+		}
+
+		const bool bTransmitUEAudio = !UPixelStreaming2PluginSettings::CVarWebRTCDisableTransmitAudio.GetValueOnAnyThread();
+		const bool bReceiveBrowserAudio = !UPixelStreaming2PluginSettings::CVarWebRTCDisableReceiveAudio.GetValueOnAnyThread();
+		if (bTransmitUEAudio || bReceiveBrowserAudio)
+		{
+			EpicRtcMediaSourceDirection AudioDirection;
+			if (bTransmitUEAudio && bReceiveBrowserAudio)
+			{
+				AudioDirection = EpicRtcMediaSourceDirection::SendRecv;
+			}
+			else if (bTransmitUEAudio)
+			{
+				AudioDirection = EpicRtcMediaSourceDirection::SendOnly;
+			}
+			else if (bReceiveBrowserAudio)
+			{
+				AudioDirection = EpicRtcMediaSourceDirection::RecvOnly;
+			}
+			else
+			{
+				AudioDirection = EpicRtcMediaSourceDirection::RecvOnly;
+			}
+
+			FUtf8String		   AudioStreamID = GetAudioStreamID();
+			EpicRtcAudioSource AudioSource = {
+				._streamId = ToEpicRtcStringView(AudioStreamID),
+				._bitrate = 510000,
+				._channels = 2,
+				._direction = AudioDirection
+			};
+
+			ParticipantConnection->AddAudioSource(AudioSource);
+		}
+
+		if (IsSFU(ParticipantId))
+		{
+			FString			  RecvLabel(TEXT("recv-datachannel"));
+			FUtf8String		  Utf8RecvLabel = *RecvLabel;
+			EpicRtcDataSource RecvDataSource = {
+				._label = ToEpicRtcStringView(Utf8RecvLabel),
+				._maxRetransmitTime = 0,
+				._maxRetransmits = 0,
+				._isOrdered = true,
+				._protocol = EpicRtcDataSourceProtocol::Sctp,
+				._negotiated = true,
+				._transportChannelId = 1
+			};
+			ParticipantConnection->AddDataSource(RecvDataSource);
+
+			FString			  SendLabel(TEXT("send-datachannel"));
+			FUtf8String		  Utf8SendLabel = *SendLabel;
+			EpicRtcDataSource SendDataSource = {
+				._label = ToEpicRtcStringView(Utf8RecvLabel),
+				._maxRetransmitTime = 0,
+				._maxRetransmits = 0,
+				._isOrdered = true,
+				._protocol = EpicRtcDataSourceProtocol::Sctp,
+				._negotiated = true,
+				._transportChannelId = 0
+			};
+			ParticipantConnection->AddDataSource(SendDataSource);
+		}
+		else
+		{
+			EpicRtcDataSource DataSource = {
+				._label = Participant->GetId(),
+				._maxRetransmitTime = 0,
+				._maxRetransmits = 0,
+				._isOrdered = true,
+				._protocol = EpicRtcDataSourceProtocol::Sctp
+			};
+			ParticipantConnection->AddDataSource(DataSource);
+		}
+
+		EpicRtcBitrate Bitrates = {
+			._minBitrateBps = UPixelStreaming2PluginSettings::CVarWebRTCMinBitrate.GetValueOnAnyThread(),
+			._hasMinBitrateBps = true,
+			._maxBitrateBps = UPixelStreaming2PluginSettings::CVarWebRTCMaxBitrate.GetValueOnAnyThread(),
+			._hasMaxBitrateBps = true,
+			._startBitrateBps = UPixelStreaming2PluginSettings::CVarWebRTCStartBitrate.GetValueOnAnyThread(),
+			._hasStartBitrateBps = true
+		};
+
+		ParticipantConnection->SetConnectionRates(Bitrates);
+
+		ParticipantConnection->StartNegotiation();
+	}
+
+	void FEpicRtcStreamer::OnRoomLeftUpdate(const EpicRtcStringView Participant)
+	{
+		FString ParticipantId = ToString(Participant);
+		UE_LOG(LogPixelStreaming2RTC, Log, TEXT("FEpicRtcStreamer::OnRoomLeftUpdate(Participant [%s] left the room.)"), *ParticipantId);
+
+		// Remove the player
+		DeletePlayerSession(ParticipantId);
+	}
+
+	void FEpicRtcStreamer::OnAudioTrackUpdate(EpicRtcParticipantInterface* Participant, EpicRtcAudioTrackInterface* AudioTrack)
+	{
+		FString ParticipantId = ToString(Participant->GetId());
+		FString AudioTrackId = ToString(AudioTrack->GetId());
+		UE_LOG(LogPixelStreaming2RTC, Log, TEXT("FEpicRtcStreamer::OnAudioTrackUpdate(Participant [%s], AudioTrack [%s])"), *ParticipantId, *AudioTrackId);
+
+		if (FPlayerContext* PlayerContext = Players->Find(ParticipantId))
+		{
+			if (AudioTrack->IsRemote())
+			{
+				PlayerContext->AudioSink = FEpicRtcAudioSink::Create(AudioTrack);
+			}
+			else
+			{
+				PlayerContext->AudioSource = FEpicRtcAudioSource::Create(AudioTrack);
+			}
+
+			if (UPixelStreaming2Delegates* Delegates = UPixelStreaming2Delegates::Get())
+			{
+				Delegates->OnAudioTrackOpenNative.Broadcast(StreamerId, ParticipantId, (bool)AudioTrack->IsRemote());
+			}
+		}
+	}
+
+	void FEpicRtcStreamer::OnVideoTrackUpdate(EpicRtcParticipantInterface* Participant, EpicRtcVideoTrackInterface* VideoTrack)
+	{
+		FString ParticipantId = ToString(Participant->GetId());
+		FString VideoTrackId = ToString(VideoTrack->GetId());
+		UE_LOG(LogPixelStreaming2RTC, Log, TEXT("FEpicRtcStreamer::OnVideoTrackUpdate(Participant [%s], VideoTrack [%s])"), *ParticipantId, *VideoTrackId);
+
+		if (FPlayerContext* PlayerContext = Players->Find(ParticipantId))
+		{
+			if (VideoTrack->IsRemote())
+			{
+				PlayerContext->VideoSink = FEpicRtcVideoSink::Create(VideoTrack);
+			}
+			else
+			{
+				PlayerContext->VideoSource = FEpicRtcVideoSource::Create(VideoTrack, VideoCapturer, VideoSourceGroup);
+			}
+
+			if (UPixelStreaming2Delegates* Delegates = UPixelStreaming2Delegates::Get())
+			{
+				Delegates->OnVideoTrackOpenNative.Broadcast(StreamerId, ParticipantId, (bool)VideoTrack->IsRemote());
+			}
+		}
+	}
+
+	void FEpicRtcStreamer::OnDataTrackUpdate(EpicRtcParticipantInterface* Participant, EpicRtcDataTrackInterface* DataTrack)
+	{
+		FString ParticipantId = ToString(Participant->GetId());
+		FString DataTrackId = ToString(DataTrack->GetId());
+		UE_LOG(LogPixelStreaming2RTC, Log, TEXT("FEpicRtcStreamer::OnDataTrackUpdate(Participant [%s], DataTrack [%s])"), *ParticipantId, *DataTrackId);
+
+		if (FPlayerContext* PlayerContext = Players->Find(ParticipantId))
+		{
+			if (!PlayerContext->DataTrack)
+			{
+				PlayerContext->DataTrack = FEpicRtcDataTrack::Create(DataTrack, InputHandler->GetFromStreamerProtocol());
+			}
+			else
+			{
+				PlayerContext->DataTrack->SetSendTrack(DataTrack);
+			}
+		}
+	}
+
+	void FEpicRtcStreamer::OnLocalSdpUpdate(EpicRtcParticipantInterface* Participant, EpicRtcSdpInterface* Sdp)
+	{
+		FString ParticipantId = ToString(Participant->GetId());
+		FString SdpType = TEXT("");
+		switch (Sdp->GetType())
+		{
+			case EpicRtcSdpType::Offer:
+				SdpType = TEXT("Offer");
+				break;
+			case EpicRtcSdpType::Answer:
+				SdpType = TEXT("Answer");
+				break;
+		}
+		UE_LOG(LogPixelStreaming2RTC, Log, TEXT("FEpicRtcStreamer::OnLocalSdpUpdate(Participant [%s], Type [%s])"), *ParticipantId, *SdpType);
+	}
+
+	void FEpicRtcStreamer::OnRemoteSdpUpdate(EpicRtcParticipantInterface* Participant, EpicRtcSdpInterface* Sdp)
+	{
+		FString ParticipantId = ToString(Participant->GetId());
+		FString SdpType = TEXT("");
+		switch (Sdp->GetType())
+		{
+			case EpicRtcSdpType::Offer:
+				SdpType = TEXT("Offer");
+				break;
+			case EpicRtcSdpType::Answer:
+				SdpType = TEXT("Answer");
+				break;
+		}
+		UE_LOG(LogPixelStreaming2RTC, Log, TEXT("FEpicRtcStreamer::OnRemoteSdpUpdate(Participant [%s], Type [%s])"), *ParticipantId, *SdpType);
+	}
+
+	void FEpicRtcStreamer::OnRoomErrorUpdate(const EpicRtcErrorCode Error)
+	{
+		UE_LOGFMT(LogPixelStreaming2RTC, VeryVerbose, "FEpicRtcStreamer::OnRoomErrorUpdate does nothing");
+	}
+
+	void FEpicRtcStreamer::OnAudioTrackMuted(EpicRtcAudioTrackInterface* AudioTrack, EpicRtcBool bIsMuted)
+	{
+		FString PlayerId;
+		bool	bFoundPlayer = FindPlayerByAudioTrack(AudioTrack, PlayerId);
+		FString AudioTrackId = ToString(AudioTrack->GetId());
+		UE_LOG(LogPixelStreaming2RTC, Log, TEXT("FEpicRtcStreamer::OnAudioTrackMuted(AudioTrack [%s], bIsMuted[%s], PlayerId[%s])"), *AudioTrackId, bIsMuted ? TEXT("True") : TEXT("False"), *PlayerId);
+
+		if (!bFoundPlayer)
+		{
+			UE_LOG(LogPixelStreaming2RTC, Warning, TEXT("FEpicRtcStreamer::OnAudioTrackMuted(Failed to find a player for audio track [%s])"), *AudioTrackId);
+			return;
+		}
+
+		FPlayerContext* PlayerContext = Players->Find(PlayerId);
+		if (AudioTrack->IsRemote())
+		{
+			PlayerContext->AudioSink->SetMuted((bool)bIsMuted);
+		}
+		else
+		{
+			PlayerContext->AudioSource->SetMuted((bool)bIsMuted);
+		}
+	}
+
+	void FEpicRtcStreamer::OnAudioTrackFrame(EpicRtcAudioTrackInterface* AudioTrack, const EpicRtcAudioFrame& Frame)
+	{
+		FString PlayerId;
+		bool	bFoundPlayer = FindPlayerByAudioTrack(AudioTrack, PlayerId);
+		FString AudioTrackId = ToString(AudioTrack->GetId());
+
+		if (!bFoundPlayer)
+		{
+			UE_LOG(LogPixelStreaming2RTC, Warning, TEXT("FEpicRtcStreamer::OnAudioTrackFrame(Failed to find a player for audio track [%s])"), *AudioTrackId);
+			return;
+		}
+
+		FPlayerContext* PlayerContext = Players->Find(PlayerId);
+		if (PlayerContext->AudioSink)
+		{
+			PlayerContext->AudioSink->OnAudioData(Frame._data, Frame._length, Frame._format._numChannels, Frame._format._sampleRate);
+		}
+	}
+
+	void FEpicRtcStreamer::OnAudioTrackRemoved(EpicRtcAudioTrackInterface* AudioTrack)
+	{
+		FString PlayerId;
+		bool	bFoundPlayer = FindPlayerByAudioTrack(AudioTrack, PlayerId);
+		FString AudioTrackId = ToString(AudioTrack->GetId());
+		UE_LOG(LogPixelStreaming2RTC, Log, TEXT("FEpicRtcStreamer::OnAudioTrackRemoved(AudioTrack [%s], PlayerId[%s])"), *AudioTrackId, *PlayerId);
+
+		if (!bFoundPlayer)
+		{
+			UE_LOG(LogPixelStreaming2RTC, Warning, TEXT("FEpicRtcStreamer::OnAudioTrackFrame(Failed to find a player for audio track [%s])"), *AudioTrackId);
+			return;
+		}
+
+		FPlayerContext* PlayerContext = Players->Find(PlayerId);
+		if (AudioTrack->IsRemote())
+		{
+			PlayerContext->AudioSink = nullptr;
+		}
+		else
+		{
+			PlayerContext->AudioSource = nullptr;
+		}
+
+		if (UPixelStreaming2Delegates* Delegates = UPixelStreaming2Delegates::Get())
+		{
+			Delegates->OnAudioTrackClosedNative.Broadcast(StreamerId, PlayerId, (bool)AudioTrack->IsRemote());
+		}
+	}
+
+	void FEpicRtcStreamer::OnAudioTrackState(EpicRtcAudioTrackInterface* AudioTrack, const EpicRtcTrackState State)
+	{
+		UE_LOGFMT(LogPixelStreaming2RTC, VeryVerbose, "FEpicRtcStreamer::OnAudioTrackState does nothing");
+	}
+
+	void FEpicRtcStreamer::OnVideoTrackMuted(EpicRtcVideoTrackInterface* VideoTrack, EpicRtcBool bIsMuted)
+	{
+		FString PlayerId;
+		bool	bFoundPlayer = FindPlayerByVideoTrack(VideoTrack, PlayerId);
+		FString VideoTrackId = ToString(VideoTrack->GetId());
+		UE_LOG(LogPixelStreaming2RTC, Log, TEXT("FEpicRtcStreamer::OnVideoTrackMuted(VideoTrack [%s], bIsMuted[%s], PlayerId[%s])"), *VideoTrackId, bIsMuted ? TEXT("True") : TEXT("False"), *PlayerId);
+
+		if (!bFoundPlayer)
+		{
+			UE_LOG(LogPixelStreaming2RTC, Warning, TEXT("FEpicRtcStreamer::OnVideoTrackMuted(Failed to find a player for video track [%s])"), *VideoTrackId);
+			return;
+		}
+
+		FPlayerContext* PlayerContext = Players->Find(PlayerId);
+		if (VideoTrack->IsRemote())
+		{
+			PlayerContext->VideoSink->SetMuted((bool)bIsMuted);
+		}
+		else
+		{
+			PlayerContext->VideoSource->SetMuted((bool)bIsMuted);
+		}
+	}
+
+	void FEpicRtcStreamer::OnVideoTrackFrame(EpicRtcVideoTrackInterface* VideoTrack, const EpicRtcVideoFrame& Frame)
+	{
+		FString PlayerId;
+		bool	bFoundPlayer = FindPlayerByVideoTrack(VideoTrack, PlayerId);
+		FString VideoTrackId = ToString(VideoTrack->GetId());
+
+		if (!bFoundPlayer)
+		{
+			UE_LOG(LogPixelStreaming2RTC, Warning, TEXT("FEpicRtcStreamer::OnVideoTrackFrame(Failed to find a player for video track [%s])"), *VideoTrackId);
+			return;
+		}
+
+		FPlayerContext* PlayerContext = Players->Find(PlayerId);
+		if (PlayerContext->VideoSink)
+		{
+			PlayerContext->VideoSink->OnVideoData(Frame);
+		}
+	}
+
+	void FEpicRtcStreamer::OnVideoTrackRemoved(EpicRtcVideoTrackInterface* VideoTrack)
+	{
+		FString PlayerId;
+		bool	bFoundPlayer = FindPlayerByVideoTrack(VideoTrack, PlayerId);
+		FString VideoTrackId = ToString(VideoTrack->GetId());
+
+		if (!bFoundPlayer)
+		{
+			UE_LOG(LogPixelStreaming2RTC, Warning, TEXT("FEpicRtcStreamer::OnVideoTrackRemoved(Failed to find a player for video track [%s])"), *VideoTrackId);
+			return;
+		}
+
+		UE_LOG(LogPixelStreaming2RTC, Log, TEXT("FEpicRtcStreamer::OnVideoTrackRemoved(VideoTrack=[%s], Player=[%s])"), *VideoTrackId, *PlayerId);
+
+		// If we did find a player, clear its video sink/source
+		FPlayerContext* PlayerContext = Players->Find(PlayerId);
+		if (VideoTrack->IsRemote())
+		{
+			PlayerContext->VideoSink = nullptr;
+		}
+		else
+		{
+			PlayerContext->VideoSource = nullptr;
+		}
+
+		if (UPixelStreaming2Delegates* Delegates = UPixelStreaming2Delegates::Get())
+		{
+			Delegates->OnVideoTrackClosedNative.Broadcast(StreamerId, PlayerId, (bool)VideoTrack->IsRemote());
+		}
+	}
+
+	void FEpicRtcStreamer::OnVideoTrackState(EpicRtcVideoTrackInterface* VideoTrack, const EpicRtcTrackState State)
+	{
+		FString PlayerId;
+		bool	bFoundPlayer = FindPlayerByVideoTrack(VideoTrack, PlayerId);
+		FString VideoTrackId = ToString(VideoTrack->GetId());
+
+		// Note: It is acceptable to not have a found a player for track state changes, as these can trigger before we have added a participant
+
+		if (State == EpicRtcTrackState::Active)
+		{
+			UE_LOG(LogPixelStreaming2RTC, Log, TEXT("FEpicRtcStreamer::OnVideoTrackState(VideoTrack=[%s], Player=[%s], State=Active)"), *VideoTrackId, *PlayerId);
+		}
+		else if (State == EpicRtcTrackState::Stopped)
+		{
+			UE_LOG(LogPixelStreaming2RTC, Log, TEXT("FEpicRtcStreamer::OnVideoTrackState(VideoTrack=[%s], Player=[%s], State=Stopped)"), *VideoTrackId, *PlayerId);
+		}
+	}
+
+	void FEpicRtcStreamer::OnDataTrackRemoved(EpicRtcDataTrackInterface* DataTrack)
+	{
+		// As long as DataTrack emits `stopped` state when it removed, this should be enough.
+		// because OnDataTrackState(DataTrack, EpicRtcTrackState::Stopped) already calls `OnDataChannelClosed`.
+	}
+
+	void FEpicRtcStreamer::OnDataTrackState(EpicRtcDataTrackInterface* DataTrack, const EpicRtcTrackState State)
+	{
+		FString PlayerId;
+		bool	bFoundPlayer = FindPlayerByDataTrack(DataTrack, PlayerId);
+		FString DataTrackId = ToString(DataTrack->GetId());
+
+		if (State == EpicRtcTrackState::Active)
+		{
+			if (!bFoundPlayer)
+			{
+				UE_LOG(LogPixelStreaming2RTC, Log, TEXT("FEpicRtcStreamer::OnDataTrackState(Message was State=Active. Failed to find a player for data track [%s])"), *DataTrackId);
+				return;
+			}
+			UE_LOG(LogPixelStreaming2RTC, Log, TEXT("FEpicRtcStreamer::OnDataTrackState(Message was State=Active. Player [%s])"), *DataTrackId);
+			OnDataChannelOpen(PlayerId);
+		}
+		else if (State == EpicRtcTrackState::Stopped)
+		{
+			UE_LOG(LogPixelStreaming2RTC, Log, TEXT("FEpicRtcStreamer::OnDataTrackState(Message was State=Stopped. Player [%s])"), *DataTrackId);
+			OnDataChannelClosed(PlayerId);
+		}
+	}
+
+	void FEpicRtcStreamer::OnDataTrackMessage(EpicRtcDataTrackInterface* DataTrack)
+	{
+		FString									DataTrackId = ToString(DataTrack->GetId());
+		TRefCountPtr<EpicRtcDataFrameInterface> DataFrame;
+		if (!DataTrack->PopFrame(DataFrame.GetInitReference()))
+		{
+			UE_LOG(LogPixelStreaming2RTC, Log, TEXT("FEpicRtcStreamer::OnDataTrackMessage(Failed to PopFrame [%s])"), *DataTrackId);
+			return;
+		}
+		FString									 PlayerId;
+		const uint8_t*							 Data = DataFrame->Data();
+		uint32_t								 DataSize = DataFrame->Size();
+		uint8									 Type = Data[0];
+		TSharedPtr<IPixelStreaming2DataProtocol> ToStreamerProtocol = InputHandler->GetToStreamerProtocol();
+		if (Type == ToStreamerProtocol->Find(EPixelStreaming2ToStreamerMessage::Multiplexed)->GetID())
+		{
+			// skip type
+			Data++;
+			DataSize--;
+			PlayerId = ReadString(Data, DataSize);
+			Type = Data[0];
+			UE_LOG(LogPixelStreaming2RTC, VeryVerbose, TEXT("FEpicRtcStreamer::OnDataTrackMessage(Received multiplexed message of type [%d] with PlayerId [%s])"), Type, *PlayerId);
+		}
+		else if (Type == ToStreamerProtocol->Find(EPixelStreaming2ToStreamerMessage::ChannelRelayStatus)->GetID())
+		{
+			HandleRelayStatusMessage(Data, DataSize, DataTrack);
+			return;
+		}
+		else if (!FindPlayerByDataTrack(DataTrack, PlayerId))
+		{
+			UE_LOG(LogPixelStreaming2RTC, Log, TEXT("FEpicRtcStreamer::OnDataTrackMessage(Failed to find a player for data track [%s])"), *DataTrackId);
+			return;
+		}
+
+		if (Type == ToStreamerProtocol->Find(EPixelStreaming2ToStreamerMessage::LatencyTest)->GetID())
+		{
+			SendLatencyReport(PlayerId);
+		}
+		else if (Type == ToStreamerProtocol->Find(EPixelStreaming2ToStreamerMessage::RequestInitialSettings)->GetID())
+		{
+			SendInitialSettings(PlayerId);
+		}
+		else if (Type == ToStreamerProtocol->Find(EPixelStreaming2ToStreamerMessage::IFrameRequest)->GetID())
+		{
+			ForceKeyFrame();
+		}
+		else if (Type == ToStreamerProtocol->Find(EPixelStreaming2ToStreamerMessage::TestEcho)->GetID())
+		{
+			if (FPlayerContext* PlayerContext = Players->Find(PlayerId))
+			{
+				if (PlayerContext->DataTrack)
+				{
+					const size_t  DescriptorSize = (DataSize - 1) / sizeof(TCHAR);
+					const TCHAR*  DescPtr = reinterpret_cast<const TCHAR*>(Data + 1);
+					const FString Message(DescriptorSize, DescPtr);
+					PlayerContext->DataTrack->SendMessage(EPixelStreaming2FromStreamerMessage::TestEcho, Message);
+				}
+			}
+		}
+		else if (!IsEngineExitRequested())
+		{
+			// If we are in "Host" mode and the current peer is not the host, then discard this input.
+			if (GetEnumFromCVar<EInputControllerMode>(UPixelStreaming2PluginSettings::CVarInputController) == EInputControllerMode::Host && InputControllingId != PlayerId)
+			{
+				return;
+			}
+
+			TArray<uint8> MessageData(Data, DataSize);
+			if (InputHandler)
+			{
+				InputHandler->OnMessage(PlayerId, MessageData);
+			}
+		}
+	}
+
+	bool FEpicRtcStreamer::FindPlayerByAudioTrack(EpicRtcAudioTrackInterface* AudioTrack, FString& OutPlayerId)
+	{
+		OutPlayerId = "";
+		// Find the player corresponding to this data track
+		FString AudioTrackId = ToString(AudioTrack->GetId());
+		Players->ApplyUntil([this, AudioTrackId, &OutPlayerId](FString DataPlayerId, FPlayerContext& PlayerContext) {
+			if (PlayerContext.AudioSource)
+			{
+				FString TrackId = ToString(PlayerContext.AudioSource->GetTrackId());
+				if (TrackId == AudioTrackId)
+				{
+					OutPlayerId = DataPlayerId;
+					return true;
+				}
+			}
+			if (PlayerContext.AudioSink)
+			{
+				FString TrackId = ToString(PlayerContext.AudioSink->GetTrackId());
+				if (TrackId == AudioTrackId)
+				{
+					OutPlayerId = DataPlayerId;
+					return true;
+				}
+			}
+			return false;
+		});
+
+		return !OutPlayerId.IsEmpty();
+	}
+
+	bool FEpicRtcStreamer::FindPlayerByVideoTrack(EpicRtcVideoTrackInterface* VideoTrack, FString& OutPlayerId)
+	{
+		OutPlayerId = "";
+		// Find the player corresponding to this data track
+		FString VideoTrackId = ToString(VideoTrack->GetId());
+		Players->ApplyUntil([this, VideoTrackId, &OutPlayerId](FString DataPlayerId, FPlayerContext& PlayerContext) {
+			if (PlayerContext.VideoSource)
+			{
+				FString TrackId = ToString(PlayerContext.VideoSource->GetTrackId());
+				if (TrackId == VideoTrackId)
+				{
+					OutPlayerId = DataPlayerId;
+					return true;
+				}
+			}
+			if (PlayerContext.VideoSink)
+			{
+				FString TrackId = ToString(PlayerContext.VideoSink->GetTrackId());
+				if (TrackId == VideoTrackId)
+				{
+					OutPlayerId = DataPlayerId;
+					return true;
+				}
+			}
+			return false;
+		});
+
+		return !OutPlayerId.IsEmpty();
+	}
+
+	bool FEpicRtcStreamer::FindPlayerByDataTrack(EpicRtcDataTrackInterface* DataTrack, FString& OutPlayerId)
+	{
+		OutPlayerId = "";
+		// Find the player corresponding to this data track
+		FString DataTrackId = ToString(DataTrack->GetId());
+		Players->ApplyUntil([this, DataTrackId, &OutPlayerId](FString DataPlayerId, FPlayerContext& PlayerContext) {
+			if (PlayerContext.DataTrack)
+			{
+				FString TrackId = ToString(PlayerContext.DataTrack->GetId());
+				if (TrackId == DataTrackId)
+				{
+					OutPlayerId = DataPlayerId;
+					return true;
+				}
+			}
+			return false;
+		});
+
+		return !OutPlayerId.IsEmpty();
+	}
+
+	FRTCStreamerFactory::FRTCStreamerFactory(TRefCountPtr<EpicRtcConferenceInterface> Conference)
+		: Conference(Conference)
+	{
+	}
+
+	FString FRTCStreamerFactory::GetStreamType()
+	{
+		return RTC_STREAM_TYPE;
+	}
+
+	TSharedPtr<IPixelStreaming2Streamer> FRTCStreamerFactory::CreateNewStreamer(const FString& StreamerId)
+	{
+		return MakeShared<FEpicRtcStreamer>(StreamerId, Conference);
+	}
+} // namespace UE::PixelStreaming2
diff --git a/Engine/Plugins/Media/PixelStreaming2/Source/PixelStreaming2RTC/Private/EpicRtcStreamer.h b/Engine/Plugins/Media/PixelStreaming2/Source/PixelStreaming2RTC/Private/EpicRtcStreamer.h
new file mode 100644
--- /dev/null
+++ b/Engine/Plugins/Media/PixelStreaming2/Source/PixelStreaming2RTC/Private/EpicRtcStreamer.h
@@ -0,0 +1,174 @@
+// Copyright Epic Games, Inc. All Rights Reserved.
+
+#pragma once
+
+#include "IPixelStreaming2Streamer.h"
+#include "VideoSourceGroup.h"
+#include "ThreadSafeMap.h"
+#include "Dom/JsonObject.h"
+#include "IPixelStreaming2InputHandler.h"
+#include "Templates/SharedPointer.h"
+#include "PlayerContext.h"
+#include "FreezeFrame.h"
+#include "StreamerReconnectTimer.h"
+#include "EpicRtcManager.h"
+
+#include "epic_rtc/core/stats.h"
+
+class IPixelStreaming2RTCModule;
+
+namespace UE::PixelStreaming2
+{
+	static const FString INVALID_PLAYER_ID = TEXT("Invalid Player Id");
+	static const FString RTC_STREAM_TYPE = TEXT("DefaultRtc");
+
+	class FEpicRtcStreamer : public IPixelStreaming2Streamer, public TSharedFromThis<FEpicRtcStreamer>
+	{
+	public:
+		FEpicRtcStreamer(const FString& StreamerId, TRefCountPtr<EpicRtcConferenceInterface> Conference);
+		static TSharedPtr<FEpicRtcStreamer> Create(const FString& StreamerId, TRefCountPtr<EpicRtcConferenceInterface> Conference);
+		virtual ~FEpicRtcStreamer();
+
+		virtual void  Initialize() override;
+		virtual void  SetStreamFPS(int32 InFramesPerSecond) override;
+		virtual int32 GetStreamFPS() override;
+		virtual void  SetCoupleFramerate(bool bCouple) override;
+
+		virtual void									SetVideoProducer(TSharedPtr<IPixelStreaming2VideoProducer> Input) override;
+		virtual TWeakPtr<IPixelStreaming2VideoProducer> GetVideoProducer() override;
+
+		virtual void	SetConnectionURL(const FString& InConnectionURL) override;
+		virtual FString GetConnectionURL() override;
+
+		virtual FString GetId() override { return StreamerId; };
+		virtual bool	IsConnected() override { return bSignallingConnected; }
+		virtual void	StartStreaming() override;
+		virtual void	StopStreaming() override;
+		virtual bool	IsStreaming() const override { return bStreamingStarted; }
+
+		virtual FPreConnectionEvent&	OnPreConnection() override;
+		virtual FStreamingStartedEvent& OnStreamingStarted() override;
+		virtual FStreamingStoppedEvent& OnStreamingStopped() override;
+
+		virtual void ForceKeyFrame() override;
+
+		virtual void FreezeStream(UTexture2D* Texture) override;
+		virtual void UnfreezeStream() override;
+
+		virtual void			SendAllPlayersMessage(FString MessageType, const FString& Descriptor) override;
+		virtual void			SendPlayerMessage(FString PlayerId, FString MessageType, const FString& Descriptor) override;
+		virtual void			SendFileData(const TArray64<uint8>& ByteData, FString& MimeType, FString& FileExtension) override;
+		virtual void			KickPlayer(FString PlayerId) override;
+		virtual TArray<FString> GetConnectedPlayers() override;
+
+		virtual TWeakPtr<IPixelStreaming2InputHandler> GetInputHandler() override { return InputHandler; }
+
+		virtual IPixelStreaming2AudioSink* GetPeerAudioSink(FString PlayerId) override;
+		virtual IPixelStreaming2AudioSink* GetUnlistenedAudioSink() override;
+		virtual IPixelStreaming2VideoSink* GetPeerVideoSink(FString PlayerId) override;
+		virtual IPixelStreaming2VideoSink* GetUnwatchedVideoSink() override;
+
+		virtual void SetConfigOption(const FName& OptionName, const FString& Value) override;
+		virtual bool GetConfigOption(const FName& OptionName, FString& OutValue) override;
+
+		virtual void PlayerRequestsBitrate(FString PlayerId, int MinBitrate, int MaxBitrate) override;
+
+		virtual void RefreshStreamBitrate() override;
+
+		void ForEachPlayer(const TFunction<void(FString, FPlayerContext)>& Func);
+
+	private:
+		// own methods
+		void OnProtocolUpdated();
+		void ConsumeStats(FString PlayerId, FName StatName, float StatValue);
+		void DeletePlayerSession(FString PlayerId);
+		void DeleteAllPlayerSessions();
+		void OnDataChannelOpen(FString PlayerId);
+		void OnDataChannelClosed(FString PlayerId);
+		void SendInitialSettings(FString PlayerId) const;
+		void SendProtocol(FString PlayerId) const;
+		void SendPeerControllerMessages(FString PlayerId) const;
+		void SendLatencyReport(FString PlayerId) const;
+		void HandleRelayStatusMessage(const uint8_t* Data, uint32_t Size, EpicRtcDataTrackInterface* DataTrack);
+		bool FindPlayerByAudioTrack(EpicRtcAudioTrackInterface* AudioTrack, FString& OutPlayerId);
+		bool FindPlayerByVideoTrack(EpicRtcVideoTrackInterface* VideoTrack, FString& OutPlayerId);
+		bool FindPlayerByDataTrack(EpicRtcDataTrackInterface* DataTrack, FString& OutPlayerId);
+		void TriggerMouseLeave(FString InStreamerId);
+		void RemoveSession(bool bDisconnect);
+		void RemoveRoom();
+		void OnStatsReady(const FString& PlayerId, const EpicRtcConnectionStats& ConnectionStats);
+
+	private:
+		FString StreamerId;
+		FString CurrentSignallingServerURL;
+
+		TSharedPtr<IPixelStreaming2InputHandler> InputHandler;
+
+		TSharedPtr<TThreadSafeMap<FString, FPlayerContext>> Players;
+
+		FString InputControllingId = INVALID_PLAYER_ID;
+
+		bool bSignallingConnected = false;
+		bool bStreamingStarted = false;
+
+		FPreConnectionEvent	   StreamingPreConnectionEvent;
+		FStreamingStartedEvent StreamingStartedEvent;
+		FStreamingStoppedEvent StreamingStoppedEvent;
+
+		TSharedPtr<FEpicRtcVideoCapturer> VideoCapturer;
+		TSharedPtr<FVideoSourceGroup>	  VideoSourceGroup;
+		TSharedPtr<FFreezeFrame>		  FreezeFrame;
+
+		FDelegateHandle ConsumeStatsHandle;
+		FDelegateHandle AllConnectionsClosedHandle;
+
+		TMap<FName, FString> ConfigOptions;
+
+		TSharedPtr<FStreamerReconnectTimer> ReconnectTimer;
+
+		TSharedPtr<FEpicRtcManager> EpicRtcManager;
+
+	public:
+		// Begin FEpicRtcManager Callbacks
+		void OnSessionStateUpdate(const EpicRtcSessionState StateUpdate);
+		void OnSessionErrorUpdate(const EpicRtcErrorCode ErrorUpdate);
+		void OnSessionRoomsAvailableUpdate(EpicRtcStringArrayInterface* RoomsList);
+
+		void OnRoomStateUpdate(const EpicRtcRoomState State);
+		void OnRoomJoinedUpdate(EpicRtcParticipantInterface* Participant);
+		void OnRoomLeftUpdate(const EpicRtcStringView ParticipantId);
+		void OnAudioTrackUpdate(EpicRtcParticipantInterface* Participant, EpicRtcAudioTrackInterface* AudioTrack);
+		void OnVideoTrackUpdate(EpicRtcParticipantInterface* Participant, EpicRtcVideoTrackInterface* VideoTrack);
+		void OnDataTrackUpdate(EpicRtcParticipantInterface* Participant, EpicRtcDataTrackInterface* DataTrack);
+		void OnLocalSdpUpdate(EpicRtcParticipantInterface* Participant, EpicRtcSdpInterface* Sdp);
+		void OnRemoteSdpUpdate(EpicRtcParticipantInterface* Participant, EpicRtcSdpInterface* Sdp);
+		void OnRoomErrorUpdate(const EpicRtcErrorCode Error);
+
+		void OnAudioTrackMuted(EpicRtcAudioTrackInterface* AudioTrack, EpicRtcBool bIsMuted);
+		void OnAudioTrackFrame(EpicRtcAudioTrackInterface* AudioTrack, const EpicRtcAudioFrame& Frame);
+		void OnAudioTrackRemoved(EpicRtcAudioTrackInterface* AudioTrack);
+		void OnAudioTrackState(EpicRtcAudioTrackInterface* AudioTrack, const EpicRtcTrackState State);
+
+		void OnVideoTrackMuted(EpicRtcVideoTrackInterface* VideoTrack, EpicRtcBool bIsMuted);
+		void OnVideoTrackFrame(EpicRtcVideoTrackInterface* VideoTrack, const EpicRtcVideoFrame& Frame);
+		void OnVideoTrackRemoved(EpicRtcVideoTrackInterface* VideoTrack);
+		void OnVideoTrackState(EpicRtcVideoTrackInterface* VideoTrack, const EpicRtcTrackState State);
+
+		void OnDataTrackRemoved(EpicRtcDataTrackInterface* DataTrack);
+		void OnDataTrackState(EpicRtcDataTrackInterface* DataTrack, const EpicRtcTrackState State);
+		void OnDataTrackMessage(EpicRtcDataTrackInterface* DataTrack);
+		// End FEpicRtcManager Callbacks
+	};
+
+	class FRTCStreamerFactory : public IPixelStreaming2StreamerFactory
+	{
+	public:
+		FRTCStreamerFactory(TRefCountPtr<EpicRtcConferenceInterface> Conference);
+
+		virtual FString								 GetStreamType() override;
+		virtual TSharedPtr<IPixelStreaming2Streamer> CreateNewStreamer(const FString& StreamerId) override;
+
+	private:
+		TRefCountPtr<EpicRtcConferenceInterface> Conference;
+	};
+} // namespace UE::PixelStreaming2
diff --git a/Engine/Plugins/Media/PixelStreaming2/Source/PixelStreaming2RTC/Private/EpicRtcVideoBufferI420.h b/Engine/Plugins/Media/PixelStreaming2/Source/PixelStreaming2RTC/Private/EpicRtcVideoBufferI420.h
new file mode 100644
--- /dev/null
+++ b/Engine/Plugins/Media/PixelStreaming2/Source/PixelStreaming2RTC/Private/EpicRtcVideoBufferI420.h
@@ -0,0 +1,50 @@
+// Copyright Epic Games, Inc. All Rights Reserved.
+
+#pragma once
+
+#include "PixelCaptureBufferI420.h"
+
+#include "epic_rtc/core/video/video_buffer.h"
+#include "epic_rtc_helper/memory/ref_count_impl_helper.h"
+
+namespace UE::PixelStreaming2
+{
+	class FEpicRtcVideoBufferI420 : public EpicRtcVideoBufferInterface, public TRefCountingMixin<FEpicRtcVideoBufferI420>
+	{
+	public:
+		FEpicRtcVideoBufferI420(TSharedPtr<FPixelCaptureBufferI420> Buffer)
+			: Buffer(Buffer) {}
+		virtual ~FEpicRtcVideoBufferI420() = default;
+
+	public:
+		// Begin EpicRtcVideoBufferInterface
+		virtual void* GetData() override
+		{
+			return Buffer->GetMutableData();
+		}
+
+		virtual EpicRtcPixelFormat GetFormat() override
+		{
+			return EpicRtcPixelFormat::I420;
+		}
+
+		virtual int GetWidth() override
+		{
+			return Buffer->GetWidth();
+		}
+
+		virtual int GetHeight() override
+		{
+			return Buffer->GetHeight();
+		}
+		// End EpicRtcVideoBufferInterface
+
+	public:
+		// Begin EpicRtcRefCountInterface
+		EPICRTC_REFCOUNT_INTERFACE_IN_PLACE
+		// End EpicRtcRefCountInterface
+
+	private:
+		TSharedPtr<FPixelCaptureBufferI420> Buffer;
+	};
+} // namespace UE::PixelStreaming2
\ No newline at end of file
diff --git a/Engine/Plugins/Media/PixelStreaming2/Source/PixelStreaming2RTC/Private/EpicRtcVideoBufferMultiFormat.cpp b/Engine/Plugins/Media/PixelStreaming2/Source/PixelStreaming2RTC/Private/EpicRtcVideoBufferMultiFormat.cpp
new file mode 100644
--- /dev/null
+++ b/Engine/Plugins/Media/PixelStreaming2/Source/PixelStreaming2RTC/Private/EpicRtcVideoBufferMultiFormat.cpp
@@ -0,0 +1,70 @@
+// Copyright Epic Games, Inc. All Rights Reserved.
+
+#include "EpicRtcVideoBufferMultiFormat.h"
+
+namespace UE::PixelStreaming2
+{
+    FEpicRtcVideoBufferMultiFormatBase::FEpicRtcVideoBufferMultiFormatBase(TSharedPtr<FPixelCaptureCapturerMultiFormat> InFrameCapturer)
+		: FrameCapturer(InFrameCapturer)
+	{
+	}
+
+	FEpicRtcVideoBufferMultiFormatLayered::FEpicRtcVideoBufferMultiFormatLayered(TSharedPtr<FPixelCaptureCapturerMultiFormat> InFrameCapturer)
+		: FEpicRtcVideoBufferMultiFormatBase(InFrameCapturer)
+	{
+	}
+
+	int FEpicRtcVideoBufferMultiFormatLayered::GetWidth()
+	{
+		return FrameCapturer ? FrameCapturer->GetWidth(GetNumLayers() - 1) : -1;
+	}
+
+	int FEpicRtcVideoBufferMultiFormatLayered::GetHeight()
+	{
+		return FrameCapturer ? FrameCapturer->GetHeight(GetNumLayers() - 1) : -1;
+	}
+
+	int FEpicRtcVideoBufferMultiFormatLayered::GetNumLayers() const
+	{
+		return FrameCapturer ? FrameCapturer->GetNumLayers() : -1;
+	}
+
+	TRefCountPtr<FEpicRtcVideoBufferMultiFormat> FEpicRtcVideoBufferMultiFormatLayered::GetLayer(int LayerIndex) const
+	{
+		return new FEpicRtcVideoBufferMultiFormat(FrameCapturer, LayerIndex);
+	}
+
+	FEpicRtcVideoBufferMultiFormat::FEpicRtcVideoBufferMultiFormat(TSharedPtr<FPixelCaptureCapturerMultiFormat> InFrameCapturer, int InLayerIndex)
+		: FEpicRtcVideoBufferMultiFormatBase(InFrameCapturer)
+		, LayerIndex(InLayerIndex)
+	{
+	}
+
+	int FEpicRtcVideoBufferMultiFormat::GetWidth()
+	{
+		return FrameCapturer ? FrameCapturer->GetWidth(LayerIndex) : -1;
+	}
+
+	int FEpicRtcVideoBufferMultiFormat::GetHeight()
+	{
+		return FrameCapturer ? FrameCapturer->GetHeight(LayerIndex) : -1;
+	}
+
+	IPixelCaptureOutputFrame* FEpicRtcVideoBufferMultiFormat::RequestFormat(int32 Format) const
+	{
+		// ensure this frame buffer will always refer to the same frame
+		if (TSharedPtr<IPixelCaptureOutputFrame>* CachedFrame = CachedFormat.Find(Format))
+		{
+			return CachedFrame->Get();
+		}
+
+		if(!FrameCapturer)
+		{
+			return nullptr;
+		}
+
+		TSharedPtr<IPixelCaptureOutputFrame> Frame = FrameCapturer->WaitForFormat(Format, LayerIndex);
+		CachedFormat.Add(Format, Frame);
+		return Frame.Get();
+	}
+}
\ No newline at end of file
diff --git a/Engine/Plugins/Media/PixelStreaming2/Source/PixelStreaming2RTC/Private/EpicRtcVideoBufferMultiFormat.h b/Engine/Plugins/Media/PixelStreaming2/Source/PixelStreaming2RTC/Private/EpicRtcVideoBufferMultiFormat.h
new file mode 100644
--- /dev/null
+++ b/Engine/Plugins/Media/PixelStreaming2/Source/PixelStreaming2RTC/Private/EpicRtcVideoBufferMultiFormat.h
@@ -0,0 +1,105 @@
+// Copyright Epic Games, Inc. All Rights Reserved.
+
+#pragma once
+
+#include "PixelCaptureCapturerMultiFormat.h"
+#include "Video/Resources/VideoResourceRHI.h"
+
+#include "epic_rtc/core/video/video_buffer.h"
+#include "epic_rtc_helper/memory/ref_count_impl_helper.h"
+
+namespace UE::PixelStreaming2
+{
+	class FEpicRtcVideoBufferMultiFormat;
+
+	/**
+	 * Base class for our multi format buffers.
+	 */
+	class FEpicRtcVideoBufferMultiFormatBase : public EpicRtcVideoBufferInterface
+	{
+	public:
+		FEpicRtcVideoBufferMultiFormatBase(TSharedPtr<FPixelCaptureCapturerMultiFormat> InFrameCapturer);
+		virtual ~FEpicRtcVideoBufferMultiFormatBase() = default;
+
+	public:
+		// Begin EpicRtcVideoBufferInterface
+		virtual void* GetData() override
+		{
+			unimplemented();
+			return nullptr;
+		}
+
+		virtual EpicRtcPixelFormat GetFormat() override
+		{
+			return EpicRtcPixelFormat::Native;
+		}
+		// End EpicRtcVideoBufferInterface
+
+	public:
+		// Begin EpicRtcRefCountInterface
+		EPICRTC_REFCOUNT_INTERFACE_IN_PLACE
+		// End EpicRtcRefCountInterface
+
+	protected:
+		TSharedPtr<FPixelCaptureCapturerMultiFormat> FrameCapturer;
+	};
+
+	/**
+	 * A multi layered, multi format frame buffer for our encoder.
+	 */
+	class FEpicRtcVideoBufferMultiFormatLayered : public FEpicRtcVideoBufferMultiFormatBase
+	{
+	public:
+		FEpicRtcVideoBufferMultiFormatLayered(TSharedPtr<FPixelCaptureCapturerMultiFormat> InFrameCapturer);
+		virtual ~FEpicRtcVideoBufferMultiFormatLayered() = default;
+
+		virtual int GetWidth() override;
+		virtual int GetHeight() override;
+
+		int											 GetNumLayers() const;
+		TRefCountPtr<FEpicRtcVideoBufferMultiFormat> GetLayer(int LayerIndex) const;
+	};
+
+	/**
+	 * An empty frame, used for stream sharing to keep the encoder alive but do no encoding.
+	 */
+	class FEpicRtcVideoBufferEmpty : public FEpicRtcVideoBufferMultiFormatLayered
+	{
+	public:
+		FEpicRtcVideoBufferEmpty(int InWidth, int InHeight)
+			: FEpicRtcVideoBufferMultiFormatLayered(nullptr)
+			, Width(InWidth)
+			, Height(InHeight)
+		{
+		}
+
+		virtual ~FEpicRtcVideoBufferEmpty() = default;
+		virtual int GetWidth() override { return Width; };
+		virtual int GetHeight() override { return Height; };
+
+	private:
+		int Width = 0;
+		int Height = 0;
+	};
+
+	/**
+	 * A single layer, multi format frame buffer.
+	 */
+	class FEpicRtcVideoBufferMultiFormat : public FEpicRtcVideoBufferMultiFormatBase
+	{
+	public:
+		FEpicRtcVideoBufferMultiFormat(TSharedPtr<FPixelCaptureCapturerMultiFormat> InFrameCapturer, int InLayerIndex);
+		virtual ~FEpicRtcVideoBufferMultiFormat() = default;
+
+		virtual int GetWidth() override;
+		virtual int GetHeight() override;
+
+		IPixelCaptureOutputFrame* RequestFormat(int32 Format) const;
+
+	private:
+		int32 LayerIndex;
+		// we want the frame buffer to always refer to the same frame. so the first request for a format
+		// will fill this cache.
+		mutable TMap<int32, TSharedPtr<IPixelCaptureOutputFrame>> CachedFormat;
+	};
+} // namespace UE::PixelStreaming2
\ No newline at end of file
diff --git a/Engine/Plugins/Media/PixelStreaming2/Source/PixelStreaming2RTC/Private/EpicRtcVideoCapturer.cpp b/Engine/Plugins/Media/PixelStreaming2/Source/PixelStreaming2RTC/Private/EpicRtcVideoCapturer.cpp
new file mode 100644
--- /dev/null
+++ b/Engine/Plugins/Media/PixelStreaming2/Source/PixelStreaming2RTC/Private/EpicRtcVideoCapturer.cpp
@@ -0,0 +1,36 @@
+// Copyright Epic Games, Inc. All Rights Reserved.
+
+#include "EpicRtcVideoCapturer.h"
+
+#include "EpicRtcVideoBufferMultiFormat.h"
+#include "PixelStreaming2PluginSettings.h"
+
+namespace UE::PixelStreaming2
+{
+	TSharedPtr<FEpicRtcVideoCapturer> FEpicRtcVideoCapturer::Create(TSharedPtr<FVideoProducer> InVideoProducer)
+	{
+		TSharedPtr<FEpicRtcVideoCapturer> VideoCapturer = TSharedPtr<FEpicRtcVideoCapturer>(new FEpicRtcVideoCapturer(InVideoProducer));
+		if (InVideoProducer)
+		{
+			VideoCapturer->FramePushedHandle = InVideoProducer->OnFramePushed.AddSP(VideoCapturer.ToSharedRef(), &FEpicRtcVideoCapturer::OnFrame);
+		}
+
+		if (UPixelStreaming2PluginSettings::FDelegates* Delegates = UPixelStreaming2PluginSettings::Delegates())
+		{
+			VideoCapturer->SimulcastEnabledChangedHandle = Delegates->OnSimulcastEnabledChanged.AddSP(VideoCapturer.ToSharedRef(), &FEpicRtcVideoCapturer::OnSimulcastEnabledChanged);
+			VideoCapturer->CaptureUseFenceChangedHandle = Delegates->OnCaptureUseFenceChanged.AddSP(VideoCapturer.ToSharedRef(), &FEpicRtcVideoCapturer::OnCaptureUseFenceChanged);
+		}
+		return VideoCapturer;
+	}
+
+	FEpicRtcVideoCapturer::FEpicRtcVideoCapturer(TSharedPtr<FVideoProducer> VideoProducer)
+		: FVideoCapturer(VideoProducer)
+	{
+		CreateFrameCapturer();
+	}
+
+	TRefCountPtr<EpicRtcVideoBufferInterface> FEpicRtcVideoCapturer::GetFrameBuffer()
+	{
+		return new FEpicRtcVideoBufferMultiFormatLayered(FrameCapturer);
+	}
+} // namespace UE::PixelStreaming2RTC
\ No newline at end of file
diff --git a/Engine/Plugins/Media/PixelStreaming2/Source/PixelStreaming2RTC/Private/EpicRtcVideoDecoder.h b/Engine/Plugins/Media/PixelStreaming2/Source/PixelStreaming2RTC/Private/EpicRtcVideoDecoder.h
new file mode 100644
--- /dev/null
+++ b/Engine/Plugins/Media/PixelStreaming2/Source/PixelStreaming2RTC/Private/EpicRtcVideoDecoder.h
@@ -0,0 +1,104 @@
+// Copyright Epic Games, Inc. All Rights Reserved.
+
+#pragma once
+
+#include "IPixelStreaming2Stats.h"
+#include "Video/Decoders/Configs/VideoDecoderConfigAV1.h"
+#include "Video/Decoders/Configs/VideoDecoderConfigH264.h"
+#include "Video/Decoders/Configs/VideoDecoderConfigVP8.h"
+#include "Video/Decoders/Configs/VideoDecoderConfigVP9.h"
+#include "Video/Resources/VideoResourceCPU.h"
+#include "Video/Resources/VideoResourceRHI.h"
+#include "Video/VideoDecoder.h"
+
+#include "epic_rtc/core/video/video_decoder.h"
+#include "epic_rtc_helper/memory/ref_count_impl_helper.h"
+
+#include <atomic>
+#include <type_traits>
+
+namespace UE::PixelStreaming2
+{
+	template <typename TResolvableVideoResource, typename TVideoResource>
+	class TVideoResourcePool
+	{
+	public:
+		TResolvableVideoResource& GetOrCreate()
+		{
+			ON_SCOPE_EXIT
+			{
+				IPixelStreaming2Stats::Get().GraphValue(TEXT("NumDecodeResource"), Resources.Num(), 1, 0.f, 120.f);
+			};
+
+			if (RHIGetInterfaceType() == ERHIInterfaceType::Vulkan)
+			{
+				// Vulkan has lifetime issues so we always use the same resource
+				if (Resources.Num() == 0)
+				{
+					return Resources.Emplace_GetRef();
+				}
+				return Resources[0];
+			}
+
+			TResolvableVideoResource* Resource = Resources.FindByPredicate([](TResolvableVideoResource& ResolvableResource) {
+				TSharedPtr<TVideoResource>& Resolved = ResolvableResource; // Activates the TSharedPtr operator overload which will resolve the resource
+				if (!Resolved.IsValid())
+				{
+					return false;
+				}
+
+				return !Resolved->IsInUse();
+			});
+
+			if (!Resource)
+			{
+				return Resources.Emplace_GetRef();
+			}
+
+			return *Resource;
+		}
+
+	private:
+		TArray<TResolvableVideoResource> Resources;
+	};
+
+	template <std::derived_from<FVideoResource> TVideoResource>
+	class TEpicRtcVideoDecoder : public EpicRtcVideoDecoderInterface
+	{
+	public:
+		TEpicRtcVideoDecoder(EpicRtcVideoCodecInfoInterface* CodecInfo);
+
+		// Begin EpicRtcVideoDecoderInterface
+		[[nodiscard]] virtual EpicRtcStringView GetName() const override;
+		virtual EpicRtcVideoDecoderConfig		GetConfig() const override;
+		virtual EpicRtcMediaResult				SetConfig(const EpicRtcVideoDecoderConfig& VideoDecoderConfig) override;
+		virtual EpicRtcMediaResult				Decode(const EpicRtcEncodedVideoFrame& Frame) override;
+		virtual void							RegisterCallback(EpicRtcVideoDecoderCallbackInterface* Callback) override;
+		virtual void							Reset() override;
+		// End EpicRtcVideoDecoderInterface
+	private:
+		TSharedPtr<TVideoDecoder<TVideoResource>>		   Decoder;
+		TUniquePtr<FVideoDecoderConfig>					   InitialVideoConfig;
+		EpicRtcVideoDecoderConfig						   DecoderConfig;
+		TRefCountPtr<EpicRtcVideoDecoderCallbackInterface> VideoDecoderCallback;
+		TRefCountPtr<EpicRtcVideoCodecInfoInterface>	   CodecInfo;
+		uint16_t										   FrameCount;
+
+		// clang-format off
+		using ResourcePoolType = typename std::conditional<std::is_same_v<TVideoResource, FVideoResourceRHI>, TVideoResourcePool<FResolvableVideoResourceRHI, FVideoResourceRHI>, 
+								 typename std::conditional<std::is_same_v<TVideoResource, FVideoResourceCPU>, TVideoResourcePool<FResolvableVideoResourceCPU, FVideoResourceCPU>, 
+								 // We only support RHI and CPU resources so if we create one that isn't this type it will error (on purpose)
+								 void>::type>::type; 
+
+		ResourcePoolType VideoResources;
+		// clang-format on
+
+	private:
+		bool LateInitDecoder();
+
+	public:
+		// Begin EpicRtcRefCountInterface
+		EPICRTC_REFCOUNT_INTERFACE_IN_PLACE
+		// End EpicRtcRefCountInterface
+	};
+} // namespace UE::PixelStreaming2
\ No newline at end of file
diff --git a/Engine/Plugins/Media/PixelStreaming2/Source/PixelStreaming2RTC/Private/EpicRtcVideoSink.cpp b/Engine/Plugins/Media/PixelStreaming2/Source/PixelStreaming2RTC/Private/EpicRtcVideoSink.cpp
new file mode 100644
--- /dev/null
+++ b/Engine/Plugins/Media/PixelStreaming2/Source/PixelStreaming2RTC/Private/EpicRtcVideoSink.cpp
@@ -0,0 +1,130 @@
+// Copyright Epic Games, Inc. All Rights Reserved.
+
+#include "EpicRtcVideoSink.h"
+
+#include "Async/Async.h"
+#include "ColorConversion.h"
+#include "EpicRtcVideoBufferI420.h"
+#include "EpicRtcVideoBufferRHI.h"
+#include "PixelStreaming2Trace.h"
+#include "RenderTargetPool.h"
+
+namespace UE::PixelStreaming2
+{
+	TSharedPtr<FEpicRtcVideoSink> FEpicRtcVideoSink::Create(TRefCountPtr<EpicRtcVideoTrackInterface> InTrack)
+	{
+		return TSharedPtr<FEpicRtcVideoSink>(new FEpicRtcVideoSink(InTrack));
+	}
+
+	FEpicRtcVideoSink::FEpicRtcVideoSink(TRefCountPtr<EpicRtcVideoTrackInterface> InTrack)
+	{
+		Track = InTrack;
+	}
+
+	void FEpicRtcVideoSink::OnVideoData(const EpicRtcVideoFrame& Frame)
+	{
+		if (!HasVideoConsumers() || bIsMuted)
+		{
+			return;
+		}
+
+		TRACE_CPUPROFILER_EVENT_SCOPE_ON_CHANNEL_STR("FVideoSink::OnVideoData", PixelStreaming2Channel);
+
+		int32_t Width = Frame._buffer->GetWidth();
+		int32_t Height = Frame._buffer->GetHeight();
+
+		if (Frame._buffer->GetFormat() == EpicRtcPixelFormat::Native)
+		{
+			FEpicRtcVideoBufferRHI* const FrameBuffer = static_cast<FEpicRtcVideoBufferRHI*>(Frame._buffer);
+			if (FrameBuffer == nullptr)
+			{
+				return;
+			}
+
+			TSharedPtr<FVideoResourceRHI, ESPMode::ThreadSafe> VideoResource = FrameBuffer->GetVideoResource();
+			if (VideoResource->GetFormat() != EVideoFormat::BGRA)
+			{
+				VideoResource = VideoResource->TransformResource(FVideoDescriptor(EVideoFormat::BGRA, Width, Height));
+			}
+
+			auto& Raw = StaticCastSharedPtr<FVideoResourceRHI>(VideoResource)->GetRaw();
+			CallConsumeFrame(Raw.Texture);
+		}
+		else if (Frame._buffer->GetFormat() == EpicRtcPixelFormat::I420)
+		{
+			{
+				FScopeLock	  Lock(&RenderSyncContext);
+				const int32_t Size = Width * Height * 4;
+				if (Size > Buffer.Num())
+				{
+					Buffer.SetNum(Size);
+				}
+
+				FEpicRtcVideoBufferI420* const FrameBuffer = static_cast<FEpicRtcVideoBufferI420*>(Frame._buffer);
+
+				int StrideY = FrameBuffer->GetWidth();
+				int StrideUV = (FrameBuffer->GetWidth() + 1) / 2;
+
+				int DataSizeY = StrideY * FrameBuffer->GetHeight();
+				int DataSizeUV = StrideUV * ((FrameBuffer->GetHeight() + 1) / 2);
+
+				uint8_t* DataY = static_cast<uint8_t*>(FrameBuffer->GetData());
+				uint8_t* DataU = static_cast<uint8_t*>(FrameBuffer->GetData()) + DataSizeY;
+				uint8_t* DataV = static_cast<uint8_t*>(FrameBuffer->GetData()) + DataSizeY + DataSizeUV;
+
+				ConvertI420ToArgb(
+					DataY, StrideY,
+					DataU, StrideUV,
+					DataV, StrideUV,
+					Buffer.GetData(), Width * 4,
+					FrameBuffer->GetWidth(), FrameBuffer->GetHeight());
+			}
+
+			AsyncTask(ENamedThreads::GetRenderThread(), [Width, Height, this]() {
+				FScopeLock Lock(&RenderSyncContext);
+
+				const FIntPoint			  FrameSize = FIntPoint(Width, Height);
+				FRHICommandListImmediate& RHICmdList = FRHICommandListExecutor::GetImmediateCommandList();
+
+				if (!RenderTargetDescriptor.IsValid() || RenderTargetDescriptor.GetSize() != FIntVector(FrameSize.X, FrameSize.Y, 0))
+				{
+					// Create the RenderTarget descriptor
+					RenderTargetDescriptor = FPooledRenderTargetDesc::Create2DDesc(FrameSize,
+						PF_B8G8R8A8,
+						FClearValueBinding::None,
+						TexCreate_None,
+						TexCreate_RenderTargetable,
+						false);
+
+					// Update the shader resource for the 'SourceTexture'
+					FRHITextureCreateDesc RenderTargetTextureDesc =
+						FRHITextureCreateDesc::Create2D(TEXT(""), FrameSize.X, FrameSize.Y, PF_B8G8R8A8)
+							.SetClearValue(FClearValueBinding::None)
+#if PLATFORM_MAC
+							.SetFlags(ETextureCreateFlags::CPUReadback | ETextureCreateFlags::SRGB)
+							.SetInitialState(ERHIAccess::CPURead);
+#else
+							.SetFlags(ETextureCreateFlags::Dynamic | ETextureCreateFlags::ShaderResource | TexCreate_RenderTargetable | ETextureCreateFlags::SRGB)
+							.SetInitialState(ERHIAccess::SRVMask);
+#endif
+
+					SourceTexture = RHICreateTexture(RenderTargetTextureDesc);
+
+					// Find a free target-able texture from the render pool
+					GRenderTargetPool.FindFreeElement(RHICmdList,
+						RenderTargetDescriptor,
+						RenderTarget,
+						TEXT("PIXELSTEAMINGPLAYER"));
+				}
+
+				// Create the update region structure
+				const FUpdateTextureRegion2D Region(0, 0, 0, 0, FrameSize.X, FrameSize.Y);
+
+				// Set the Pixel data of the webrtc Frame to the SourceTexture
+				RHIUpdateTexture2D(SourceTexture, 0, Region, FrameSize.X * 4, Buffer.GetData());
+
+				CallConsumeFrame(SourceTexture);
+			});
+		}
+	}
+} // namespace UE::PixelStreaming2
\ No newline at end of file
diff --git a/Engine/Plugins/Media/PixelStreaming2/Source/PixelStreaming2RTC/Private/EpicRtcVideoTrack.cpp b/Engine/Plugins/Media/PixelStreaming2/Source/PixelStreaming2RTC/Private/EpicRtcVideoTrack.cpp
new file mode 100644
--- /dev/null
+++ b/Engine/Plugins/Media/PixelStreaming2/Source/PixelStreaming2RTC/Private/EpicRtcVideoTrack.cpp
@@ -0,0 +1,16 @@
+// Copyright Epic Games, Inc. All Rights Reserved.
+
+#include "EpicRtcVideoTrack.h"
+
+namespace UE::PixelStreaming2
+{
+    EpicRtcStringView FEpicRtcVideoTrack::GetTrackId() const
+    {
+        if (!Track)
+		{
+			return EpicRtcStringView { ._ptr = nullptr, ._length = 0 };
+		}
+
+		return Track->GetId();
+    }
+}
diff --git a/Engine/Plugins/Media/PixelStreaming2/Source/PixelStreaming2RTC/Private/EpicRtcVideoTrackObserver.cpp b/Engine/Plugins/Media/PixelStreaming2/Source/PixelStreaming2RTC/Private/EpicRtcVideoTrackObserver.cpp
new file mode 100644
--- /dev/null
+++ b/Engine/Plugins/Media/PixelStreaming2/Source/PixelStreaming2RTC/Private/EpicRtcVideoTrackObserver.cpp
@@ -0,0 +1,55 @@
+// Copyright Epic Games, Inc. All Rights Reserved.
+
+#include "EpicRtcVideoTrackObserver.h"
+
+#include "EpicRtcManager.h"
+
+namespace UE::PixelStreaming2
+{
+
+	FEpicRtcVideoTrackObserver::FEpicRtcVideoTrackObserver(TWeakPtr<FEpicRtcManager> Manager)
+		: Manager(Manager)
+	{
+	}
+
+	void FEpicRtcVideoTrackObserver::OnVideoTrackMuted(EpicRtcVideoTrackInterface* VideoTrack, EpicRtcBool bIsMuted)
+	{
+		if (!Manager.IsValid())
+		{
+			return;
+		}
+
+		Manager.Pin()->OnVideoTrackMuted.Broadcast(VideoTrack, bIsMuted);
+	}
+
+	void FEpicRtcVideoTrackObserver::OnVideoTrackFrame(EpicRtcVideoTrackInterface* VideoTrack, const EpicRtcVideoFrame& Frame)
+	{
+		if (!Manager.IsValid())
+		{
+			return;
+		}
+
+		Manager.Pin()->OnVideoTrackFrame.Broadcast(VideoTrack, Frame);
+	}
+
+	void FEpicRtcVideoTrackObserver::OnVideoTrackRemoved(EpicRtcVideoTrackInterface* VideoTrack)
+	{
+		if (!Manager.IsValid())
+		{
+			return;
+		}
+
+		Manager.Pin()->OnVideoTrackRemoved.Broadcast(VideoTrack);
+	}
+
+	void FEpicRtcVideoTrackObserver::OnVideoTrackState(EpicRtcVideoTrackInterface* VideoTrack, const EpicRtcTrackState State)
+	{
+		if (!Manager.IsValid())
+		{
+			return;
+		}
+
+		Manager.Pin()->OnVideoTrackState.Broadcast(VideoTrack, State);
+	}
+
+} // namespace UE::PixelStreaming2
\ No newline at end of file
