[PS2] PR #12718: Refactor PixelStreaming2 plugin to enable other streaming protocols

diff --git a/Engine/Plugins/Media/PixelStreaming2/Source/PixelStreaming2RTC/Private/PixelStreaming2RTCModule.cpp b/Engine/Plugins/Media/PixelStreaming2/Source/PixelStreaming2RTC/Private/PixelStreaming2RTCModule.cpp
new file mode 100644
--- /dev/null
+++ b/Engine/Plugins/Media/PixelStreaming2/Source/PixelStreaming2RTC/Private/PixelStreaming2RTCModule.cpp
@@ -0,0 +1,474 @@
+// Copyright Epic Games, Inc. All Rights Reserved.
+
+#include "PixelStreaming2RTCModule.h"
+
+#include "UtilsCoder.h"
+#include "CoreMinimal.h"
+#include "UtilsCore.h"
+#include "Engine/GameViewportClient.h"
+#include "Engine/Texture2D.h"
+#include "IPixelStreaming2Module.h"
+#include "IPixelStreaming2InputModule.h"
+#include "Logging.h"
+#include "Modules/ModuleManager.h"
+#include "PixelStreaming2Delegates.h"
+#include "PixelStreaming2Utils.h"
+#include "PixelStreaming2PluginSettings.h"
+#include "Slate/SceneViewport.h"
+#include "EpicRtcStreamer.h"
+#include "UObject/UObjectIterator.h"
+#include "UtilsCommon.h"
+
+#if PLATFORM_LINUX
+	#include "CudaModule.h"
+#endif
+
+#if PLATFORM_WINDOWS
+	#include "Windows/AllowWindowsPlatformTypes.h"
+THIRD_PARTY_INCLUDES_START
+	#include <VersionHelpers.h>
+THIRD_PARTY_INCLUDES_END
+	#include "Windows/HideWindowsPlatformTypes.h"
+#endif
+
+#include "AudioDeviceManager.h"
+#include "RenderingThread.h"
+#include "Misc/CommandLine.h"
+#include "Misc/Parse.h"
+#include "RendererInterface.h"
+#include "Rendering/SlateRenderer.h"
+#include "Framework/Application/SlateApplication.h"
+#include "Misc/ConfigCacheIni.h"
+#include "GameFramework/GameModeBase.h"
+#include "Dom/JsonObject.h"
+#include "Misc/App.h"
+#include "Misc/MessageDialog.h"
+#include "Async/Async.h"
+#include "Engine/Engine.h"
+#include "WebSocketsModule.h"
+
+#include "Video/Encoders/Configs/VideoEncoderConfigH264.h"
+#include "Video/Encoders/Configs/VideoEncoderConfigAV1.h"
+
+#if !UE_BUILD_SHIPPING
+	#include "DrawDebugHelpers.h"
+#endif
+
+#include "VideoProducerBackBuffer.h"
+#include "VideoProducerMediaCapture.h"
+#include "Engine/GameEngine.h"
+#include "Stats.h"
+#include "UtilsString.h"
+
+#include "EpicRtcAllocator.h"
+#include "EpicRtcAudioCapturer.h"
+#include "EpicRtcLogging.h"
+#include "EpicRtcVideoEncoderInitializer.h"
+#include "EpicRtcVideoDecoderInitializer.h"
+#include "EpicRtcWebsocketFactory.h"
+
+namespace UE::PixelStreaming2
+{
+	constexpr EpicRtcLogLevel UnrealLogToEpicRtcCategoryMap[] = {
+		EpicRtcLogLevel::Off,
+		EpicRtcLogLevel::Critical,
+		EpicRtcLogLevel::Error,
+		EpicRtcLogLevel::Warning,
+		EpicRtcLogLevel::Info,
+		EpicRtcLogLevel::Info,
+		EpicRtcLogLevel::Debug,
+		EpicRtcLogLevel::Trace,
+		EpicRtcLogLevel::Trace
+	};
+
+	static_assert(EpicRtcLogLevel::Off == UnrealLogToEpicRtcCategoryMap[ELogVerbosity::Type::NoLogging]);
+	static_assert(EpicRtcLogLevel::Critical == UnrealLogToEpicRtcCategoryMap[ELogVerbosity::Type::Fatal]);
+	static_assert(EpicRtcLogLevel::Error == UnrealLogToEpicRtcCategoryMap[ELogVerbosity::Type::Error]);
+	static_assert(EpicRtcLogLevel::Warning == UnrealLogToEpicRtcCategoryMap[ELogVerbosity::Type::Warning]);
+	static_assert(EpicRtcLogLevel::Info == UnrealLogToEpicRtcCategoryMap[ELogVerbosity::Type::Display]);
+	static_assert(EpicRtcLogLevel::Info == UnrealLogToEpicRtcCategoryMap[ELogVerbosity::Type::Log]);
+	static_assert(EpicRtcLogLevel::Debug == UnrealLogToEpicRtcCategoryMap[ELogVerbosity::Type::Verbose]);
+	static_assert(EpicRtcLogLevel::Trace == UnrealLogToEpicRtcCategoryMap[ELogVerbosity::Type::VeryVerbose]);
+	static_assert(EpicRtcLogLevel::Trace == UnrealLogToEpicRtcCategoryMap[ELogVerbosity::Type::All]);
+
+	FPixelStreaming2RTCModule* FPixelStreaming2RTCModule::PixelStreaming2Module = nullptr;
+
+	FUtf8String FPixelStreaming2RTCModule::EpicRtcConferenceName("pixel_streaming_conference_instance");
+
+	/**
+	 * Stats logger - as turned on/off by CVarPixelStreaming2LogStats
+	 */
+	void ConsumeStat(FString PlayerId, FName StatName, float StatValue)
+	{
+		UE_LOGFMT(LogPixelStreaming2RTC, Log, "[{0}]({1}) = {2}", PlayerId, StatName.ToString(), StatValue);
+	}
+
+	/**
+	 * IModuleInterface implementation
+	 */
+	void FPixelStreaming2RTCModule::StartupModule()
+	{
+#if UE_SERVER
+		// Hack to no-op the rest of the module so Blueprints can still work
+		return;
+#endif
+		if (!IsStreamingSupported())
+		{
+			return;
+		}
+
+		if (!FSlateApplication::IsInitialized())
+		{
+			return;
+		}
+
+		const ERHIInterfaceType RHIType = GDynamicRHI ? RHIGetInterfaceType() : ERHIInterfaceType::Hidden;
+		// only D3D11/D3D12/Vulkan is supported
+		if (!(RHIType == ERHIInterfaceType::D3D11 || RHIType == ERHIInterfaceType::D3D12 || RHIType == ERHIInterfaceType::Vulkan || RHIType == ERHIInterfaceType::Metal))
+		{
+#if !WITH_DEV_AUTOMATION_TESTS
+			UE_LOG(LogPixelStreaming2RTC, Warning, TEXT("Only D3D11/D3D12/Vulkan/Metal Dynamic RHI is supported. Detected %s"), GDynamicRHI != nullptr ? GDynamicRHI->GetName() : TEXT("[null]"));
+#endif
+			return;
+		}
+
+		// By calling InitDefaultStreamer post engine init we can use pixel streaming in standalone editor mode
+		IPixelStreaming2Module::Get().OnReady().AddLambda([this](IPixelStreaming2Module& CoreModule) {
+			// Need to initialize after other modules have initialized such as NVCodec.
+			if (!InitializeEpicRtc())
+			{
+				return;
+			}
+
+			if (!ensure(GEngine != nullptr))
+			{
+				return;
+			}
+
+			StreamerFactory.Reset(new FRTCStreamerFactory(EpicRtcConference));
+
+			// Ensure we have ImageWrapper loaded, used in Freezeframes
+			verify(FModuleManager::Get().LoadModule(FName("ImageWrapper")));
+
+			// We don't want to start immediately streaming in editor
+			if (!GIsEditor)
+			{
+				InitDefaultStreamer();
+				if (DefaultStreamer)
+				{
+					DefaultStreamer->StartStreaming();
+				}
+			}
+
+			bModuleReady = true;
+			ReadyEvent.Broadcast(*this);
+		});
+
+		FModuleManager::LoadModuleChecked<FWebSocketsModule>("WebSockets");
+
+		// Call these to initialize their singletons
+		FStats::Get();
+
+		// Extra Initialisations post loading console commands
+		IConsoleManager::Get().RegisterConsoleCommand(
+			TEXT("PixelStreaming.StartStreaming"),
+			TEXT("Start all streaming sessions"),
+			FConsoleCommandDelegate::CreateLambda([DefaultStreamer = this->DefaultStreamer]() {
+				if (DefaultStreamer)
+				{
+					DefaultStreamer->StartStreaming();
+				}
+			}));
+
+		IConsoleManager::Get().RegisterConsoleCommand(
+			TEXT("PixelStreaming.StopStreaming"),
+			TEXT("End any existing streaming sessions."),
+			FConsoleCommandDelegate::CreateLambda([DefaultStreamer = this->DefaultStreamer]() {
+				if (DefaultStreamer)
+				{
+					DefaultStreamer->StopStreaming();
+				}
+			}));
+
+		if (UPixelStreaming2PluginSettings::FDelegates* Delegates = UPixelStreaming2PluginSettings::Delegates())
+		{
+			Delegates->OnLogStatsChanged.AddLambda([this](IConsoleVariable* Var) {
+				bool					   bLogStats = Var->GetBool();
+				UPixelStreaming2Delegates* Delegates = UPixelStreaming2Delegates::Get();
+				if (!Delegates)
+				{
+					return;
+				}
+				if (bLogStats)
+				{
+					LogStatsHandle = Delegates->OnStatChangedNative.AddStatic(&ConsumeStat);
+				}
+				else
+				{
+					Delegates->OnStatChangedNative.Remove(LogStatsHandle);
+				}
+			});
+
+			Delegates->OnWebRTCFpsChanged.AddLambda([](IConsoleVariable*) {
+				IPixelStreaming2Module::Get().ForEachStreamer([](TSharedPtr<IPixelStreaming2Streamer> Streamer) {
+					Streamer->RefreshStreamBitrate();
+				});
+			});
+
+			Delegates->OnWebRTCBitrateChanged.AddLambda([](IConsoleVariable*) {
+				IPixelStreaming2Module::Get().ForEachStreamer([](TSharedPtr<IPixelStreaming2Streamer> Streamer) {
+					Streamer->RefreshStreamBitrate();
+				});
+			});
+		}
+
+		bStartupCompleted = true;
+	}
+
+	void FPixelStreaming2RTCModule::ShutdownModule()
+	{
+		if (!IsStreamingSupported())
+		{
+			return;
+		}
+
+		if (!bStartupCompleted)
+		{
+			return;
+		}
+
+		// We explicitly call release on streamer so WebRTC gets shutdown before our module is deleted
+		// additionally the streamer does a bunch of delegate calls and unbinds which seem to have issues
+		// when called during engine destruction rather than here.
+		DefaultStreamer.Reset();
+
+		AudioMixingCapturer.Reset();
+
+		TickConferenceTask.Reset();
+		StreamerFactory.Reset();
+
+		if (!EpicRtcPlatform)
+		{
+			UE_LOGFMT(LogPixelStreaming2RTC, Error, "EpicRtcPlatform does not exist during shutdown when it is expected to exist");
+		}
+		else
+		{
+			EpicRtcPlatform->ReleaseConference(ToEpicRtcStringView(EpicRtcConferenceName));
+		}
+
+		bStartupCompleted = false;
+	}
+
+	/**
+	 * End IModuleInterface implementation
+	 */
+
+	FPixelStreaming2RTCModule* FPixelStreaming2RTCModule::GetModule()
+	{
+		if (!PixelStreaming2Module)
+		{
+			PixelStreaming2Module = FModuleManager::Get().LoadModulePtr<FPixelStreaming2RTCModule>("PixelStreaming2RTC");
+		}
+
+		return PixelStreaming2Module;
+	}
+
+	/**
+	 * IPixelStreaming2RTCModule implementation
+	 */
+	IPixelStreaming2RTCModule::FReadyEvent& FPixelStreaming2RTCModule::OnReady()
+	{
+		return ReadyEvent;
+	}
+
+	bool FPixelStreaming2RTCModule::IsReady()
+	{
+		return bModuleReady;
+	}
+	/**
+	 * End IPixelStreaming2RTCModule implementation
+	 */
+
+	TSharedPtr<FEpicRtcAudioCapturer> FPixelStreaming2RTCModule::GetAudioCapturer()
+	{
+		if (!AudioMixingCapturer)
+		{
+			AudioMixingCapturer = FEpicRtcAudioCapturer::Create();
+		}
+
+		return AudioMixingCapturer;
+	}
+
+	void FPixelStreaming2RTCModule::InitDefaultStreamer()
+	{
+		IPixelStreaming2Module& CoreModule = IPixelStreaming2Module::Get();
+
+		FString ConnectionURL = CoreModule.GetDefaultConnectionURL();
+		FURL	URL(nullptr, *ConnectionURL, TRAVEL_Absolute);
+		if (URL.Protocol != TEXT("ws") && URL.Protocol != TEXT("wss"))
+		{
+			UE_LOGFMT(LogPixelStreaming2RTC, Log, "PixelStreaming2RTC early out due to protocol not matching ws or wss");
+		}
+
+		UE_LOGFMT(LogPixelStreaming2RTC, Log, "PixelStreaming2 streamer ID: {0}", CoreModule.GetDefaultStreamerID());
+
+		DefaultStreamer = CoreModule.CreateStreamer(CoreModule.GetDefaultStreamerID(), StreamerFactory->GetStreamType());
+
+		// default to the scene viewport if we have a game engine
+		if (UGameEngine* GameEngine = Cast<UGameEngine>(GEngine))
+		{
+			TSharedPtr<SWindow>						 TargetWindow = GameEngine->GameViewport->GetWindow();
+			TSharedPtr<IPixelStreaming2InputHandler> InputHandler = DefaultStreamer->GetInputHandler().Pin();
+			if (TargetWindow.IsValid() && InputHandler.IsValid())
+			{
+				InputHandler->SetTargetWindow(TargetWindow);
+			}
+			else
+			{
+				UE_LOG(LogPixelStreaming2RTC, Error, TEXT("Cannot set target window - target window is not valid."));
+			}
+		}
+
+		if (!DefaultStreamer->GetConnectionURL().IsEmpty())
+		{
+			// The user has specified a URL on the command line meaning their intention is to start streaming immediately
+			// in that case, set up the video input for them (as long as we're not in editor)
+			if (UPixelStreaming2PluginSettings::CVarUseMediaCapture.GetValueOnAnyThread())
+			{
+				DefaultStreamer->SetVideoProducer(FVideoProducerMediaCapture::CreateActiveViewportCapture());
+			}
+			else
+			{
+				DefaultStreamer->SetVideoProducer(FVideoProducerBackBuffer::Create());
+			}
+		}
+	}
+
+	FString FPixelStreaming2RTCModule::GetFieldTrials()
+	{
+		FString FieldTrials = UPixelStreaming2PluginSettings::CVarWebRTCFieldTrials.GetValueOnAnyThread();
+
+		// Set the WebRTC-FrameDropper/Disabled/ if the CVar is set
+		if (UPixelStreaming2PluginSettings::CVarWebRTCDisableFrameDropper.GetValueOnAnyThread())
+		{
+			FieldTrials += TEXT("WebRTC-FrameDropper/Disabled/");
+		}
+
+		if (UPixelStreaming2PluginSettings::CVarWebRTCEnableFlexFec.GetValueOnAnyThread())
+		{
+			FieldTrials += TEXT("WebRTC-FlexFEC-03-Advertised/Enabled/WebRTC-FlexFEC-03/Enabled/");
+		}
+
+		// Parse "WebRTC-Video-Pacing/" field trial
+		{
+			float PacingFactor = UPixelStreaming2PluginSettings::CVarWebRTCVideoPacingFactor.GetValueOnAnyThread();
+			float PacingMaxDelayMs = UPixelStreaming2PluginSettings::CVarWebRTCVideoPacingMaxDelay.GetValueOnAnyThread();
+
+			if (PacingFactor >= 0.0f || PacingMaxDelayMs >= 0.0f)
+			{
+				FString VideoPacingFieldTrialStr = TEXT("WebRTC-Video-Pacing/");
+				bool	bHasPacingFactor = PacingFactor >= 0.0f;
+				if (bHasPacingFactor)
+				{
+					VideoPacingFieldTrialStr += FString::Printf(TEXT("factor:%.1f"), PacingFactor);
+				}
+				bool bHasMaxDelay = PacingMaxDelayMs >= 0.0f;
+				if (bHasMaxDelay)
+				{
+					VideoPacingFieldTrialStr += bHasPacingFactor ? TEXT(",") : TEXT("");
+					VideoPacingFieldTrialStr += FString::Printf(TEXT("max_delay:%.0f"), PacingMaxDelayMs);
+				}
+				VideoPacingFieldTrialStr += TEXT("/");
+				FieldTrials += VideoPacingFieldTrialStr;
+			}
+		}
+
+		return FieldTrials;
+	}
+
+	bool FPixelStreaming2RTCModule::InitializeEpicRtc()
+	{
+		// Initialize the audio capturer (if not already called by some user code)
+		AudioMixingCapturer = GetAudioCapturer();
+
+		EpicRtcVideoEncoderInitializers = { new FEpicRtcVideoEncoderInitializer() };
+		EpicRtcVideoDecoderInitializers = { new FEpicRtcVideoDecoderInitializer() };
+
+		EpicRtcPlatformConfig PlatformConfig{
+			._memory = new FEpicRtcAllocator()
+		};
+
+		EpicRtcErrorCode Result = GetOrCreatePlatform(PlatformConfig, EpicRtcPlatform.GetInitReference());
+		if (Result != EpicRtcErrorCode::Ok && Result != EpicRtcErrorCode::FoundExistingPlatform)
+		{
+			UE_LOG(LogPixelStreaming2RTC, Warning, TEXT("Unable to create EpicRtc Platform. GetOrCreatePlatform returned %s"), *ToString(Result));
+			return false;
+		}
+
+		FUtf8String EpicRtcFieldTrials(GetFieldTrials());
+
+		WebsocketFactory = MakeRefCount<FEpicRtcWebsocketFactory>();
+
+		StatsCollector = MakeRefCount<FEpicRtcStatsCollector>();
+
+		// clang-format off
+		EpicRtcConfig ConferenceConfig = {
+			._websocketFactory = WebsocketFactory.GetReference(),
+			._signallingType = EpicRtcSignallingType::PixelStreaming,
+			._signingPlugin = nullptr,
+			._migrationPlugin = nullptr,
+			._audioDevicePlugin = nullptr,
+			._audioConfig = {
+				._tickAdm = true,
+				._audioEncoderInitializers = {}, // Not needed because we use the inbuilt audio codecs
+				._audioDecoderInitializers = {}, // Not needed because we use the inbuilt audio codecs
+				._enableBuiltInAudioCodecs = true,
+			},
+			._videoConfig = {
+				._videoEncoderInitializers = {
+					._ptr = const_cast<const EpicRtcVideoEncoderInitializerInterface**>(EpicRtcVideoEncoderInitializers.GetData()),
+					._size = (uint64_t)EpicRtcVideoEncoderInitializers.Num()
+				},
+				._videoDecoderInitializers = {
+					._ptr = const_cast<const EpicRtcVideoDecoderInitializerInterface**>(EpicRtcVideoDecoderInitializers.GetData()),
+					._size = (uint64_t)EpicRtcVideoDecoderInitializers.Num()
+				},
+				._enableBuiltInVideoCodecs = false
+			},
+			._fieldTrials = {
+				._fieldTrials = ToEpicRtcStringView(EpicRtcFieldTrials),
+				._isGlobal = 0
+			},
+			._logging = {
+				._logger = new FEpicRtcLogsRedirector(),
+#if !NO_LOGGING // When building WITH_SHIPPING by default .GetVerbosity() does not exist
+				._level = UnrealLogToEpicRtcCategoryMap[LogPixelStreaming2EpicRtc.GetVerbosity()],
+				._levelWebRtc = UnrealLogToEpicRtcCategoryMap[LogPixelStreaming2WebRtc.GetVerbosity()]
+#endif
+			},
+			._stats = {
+				._statsCollectorCallback = StatsCollector.GetReference(),
+				._statsCollectorInterval = 1000,
+				._jsonFormatOnly = false
+			}
+		};
+		// clang-format on
+
+		Result = EpicRtcPlatform->CreateConference(ToEpicRtcStringView(EpicRtcConferenceName), ConferenceConfig, EpicRtcConference.GetInitReference());
+		if (Result != EpicRtcErrorCode::Ok)
+		{
+			UE_LOG(LogPixelStreaming2RTC, Warning, TEXT("Unable to create EpicRtc Conference: CreateConference returned %s"), *ToString(Result));
+			return false;
+		}
+
+		TickConferenceTask = FPixelStreamingTickableTask::Create<FEpicRtcTickConferenceTask>(EpicRtcConference, TEXT("PixelStreaming2Module TickConferenceTask"));
+
+		return true;
+	}
+
+	/**
+	 * End own methods
+	 */
+} // namespace UE::PixelStreaming2
+
+IMPLEMENT_MODULE(UE::PixelStreaming2::FPixelStreaming2RTCModule, PixelStreaming2RTC)
diff --git a/Engine/Plugins/Media/PixelStreaming2/Source/PixelStreaming2RTC/Private/PixelStreaming2RTCModule.h b/Engine/Plugins/Media/PixelStreaming2/Source/PixelStreaming2RTC/Private/PixelStreaming2RTCModule.h
new file mode 100644
--- /dev/null
+++ b/Engine/Plugins/Media/PixelStreaming2/Source/PixelStreaming2RTC/Private/PixelStreaming2RTCModule.h
@@ -0,0 +1,79 @@
+// Copyright Epic Games, Inc. All Rights Reserved.
+
+#pragma once
+
+#include "IPixelStreaming2RTCModule.h"
+#include "EpicRtcConferenceUtils.h"
+#include "EpicRtcStatsCollector.h"
+#include "EpicRtcStreamer.h"
+
+#include "epic_rtc/core/platform.h"
+#include "epic_rtc/plugins/signalling/signalling_type.h"
+
+class UPixelStreaming2Input;
+class SWindow;
+
+namespace UE::PixelStreaming2
+{
+	class FVideoInputBackBuffer;
+	class FVideoSourceGroup;
+	class FEpicRtcWebsocketFactory;
+
+	/**
+	 * This plugin allows the back buffer to be sent as a compressed video across a network.
+	 */
+	class FPixelStreaming2RTCModule : public IPixelStreaming2RTCModule
+	{
+	public:
+		static FPixelStreaming2RTCModule* GetModule();
+
+		virtual ~FPixelStreaming2RTCModule() = default;
+
+		// Begin IPixelStreaming2RTCModule
+		virtual FReadyEvent& OnReady() override;
+		virtual bool		 IsReady() override;
+		// End IPixelStreaming2RTCModule
+
+		TSharedPtr<class FEpicRtcAudioCapturer>	 GetAudioCapturer();
+		TRefCountPtr<EpicRtcConferenceInterface> GetEpicRtcConference() { return EpicRtcConference; }
+		TRefCountPtr<FEpicRtcStatsCollector>	 GetStatsCollector() { return StatsCollector; }
+
+	private:
+		// Begin IModuleInterface
+		void StartupModule() override;
+		void ShutdownModule() override;
+		// End IModuleInterface
+
+		// Own methods
+		void InitDefaultStreamer();
+
+		FString GetFieldTrials();
+		bool	InitializeEpicRtc();
+
+	private:
+		bool							  bModuleReady = false;
+		bool							  bStartupCompleted = false;
+		static FPixelStreaming2RTCModule* PixelStreaming2Module;
+
+		FReadyEvent		ReadyEvent;
+		FDelegateHandle LogStatsHandle;
+
+		TSharedPtr<IPixelStreaming2Streamer> DefaultStreamer;
+
+	private:
+		TSharedPtr<class FEpicRtcAudioCapturer>	 AudioMixingCapturer;
+		TRefCountPtr<EpicRtcPlatformInterface>	 EpicRtcPlatform;
+		TRefCountPtr<EpicRtcConferenceInterface> EpicRtcConference;
+		TRefCountPtr<FEpicRtcStatsCollector>	 StatsCollector;
+
+		TRefCountPtr<FEpicRtcWebsocketFactory>	   WebsocketFactory;
+		TUniqueTaskPtr<FEpicRtcTickConferenceTask> TickConferenceTask;
+
+		TArray<EpicRtcVideoEncoderInitializerInterface*> EpicRtcVideoEncoderInitializers;
+		TArray<EpicRtcVideoDecoderInitializerInterface*> EpicRtcVideoDecoderInitializers;
+
+		static FUtf8String EpicRtcConferenceName;
+
+		TUniquePtr<FRTCStreamerFactory> StreamerFactory;
+	};
+} // namespace UE::PixelStreaming2
diff --git a/Engine/Plugins/Media/PixelStreaming2/Source/PixelStreaming2RTC/Private/Tests/EpicRtcUtils.cpp b/Engine/Plugins/Media/PixelStreaming2/Source/PixelStreaming2RTC/Private/Tests/EpicRtcUtils.cpp
new file mode 100644
--- /dev/null
+++ b/Engine/Plugins/Media/PixelStreaming2/Source/PixelStreaming2RTC/Private/Tests/EpicRtcUtils.cpp
@@ -0,0 +1,150 @@
+// Copyright Epic Games, Inc. All Rights Reserved.
+
+#include "EpicRtcUtils.h"
+
+namespace UE::PixelStreaming2
+{
+	bool FTickAndWaitOrTimeout::Update()
+	{
+		if (Manager->GetEpicRtcConference())
+		{
+			while (Manager->GetEpicRtcConference()->NeedsTick())
+			{
+				Manager->GetEpicRtcConference()->Tick();
+			}
+		}
+
+		if (CheckFunc())
+		{
+			return true;
+		}
+
+		double DeltaTime = FPlatformTime::Seconds() - StartTime;
+		if (DeltaTime > TimeoutSeconds)
+		{
+			UE_LOG(LogPixelStreaming2RTC, Error, TEXT("Timed out"));
+			return true;
+		}
+		return false; // Latent Test return false will run again next frame
+	}
+
+	bool FDisconnectRoom::Update()
+	{
+		TRefCountPtr<EpicRtcRoomInterface>& Room = Manager->GetEpicRtcRoom();
+
+		if (!Room)
+		{
+			UE_LOG(LogPixelStreaming2RTC, Error, TEXT("Unable to disconnect room, Room does not exist"));
+			return true;
+		}
+
+		Room->Leave();
+
+		return true;
+	}
+
+	bool FCleanupRoom::Update()
+	{
+		TRefCountPtr<EpicRtcRoomInterface>& Room = Manager->GetEpicRtcRoom();
+		if (!Room)
+		{
+			UE_LOG(LogPixelStreaming2RTC, Error, TEXT("Unable to update room, Room does not exist"));
+			return true;
+		}
+
+		Manager->GetEpicRtcSession()->RemoveRoom(ToEpicRtcStringView(RoomId));
+
+		// EpicRtc has released its hold on the room. All that should be holding a ref is the manager
+		if (Manager->GetEpicRtcRoom()->Count() != 1)
+		{
+			UE_LOG(LogPixelStreaming2RTC, Error, TEXT("Room has invalid reference count. Expected 1, Actual (%d)"), Manager->GetEpicRtcRoom()->Count());
+			return true;
+		}
+
+		// Cannot call Release on TRefCountPtr without it storing a pointer to a released object.
+		// So grab the pointer to the object, set to nullptr which will call release
+		// By directly calling release, we get the final count which we can check.
+		EpicRtcRoomInterface* RoomPtr = Room.GetReference();
+		RoomPtr->AddRef();
+		Room = nullptr; // will callrelease so count is same as before calling RoomPtr->AddRef();
+		if (uint32 Count = RoomPtr->Release(); Count != 0)
+		{
+			UE_LOG(LogPixelStreaming2RTC, Error, TEXT("Room has invalid reference count. Expected 0, Actual (%d)"), Count);
+		}
+
+		// The session has been destroyed, the only thing holding a ref to the track observer factories should be the manager
+		if (Manager->GetAudioTrackObserverFactory()->Count() != 1)
+		{
+			UE_LOG(LogPixelStreaming2RTC, Error, TEXT("AudioTrackObserverFactory has invalid reference count. Expected 1, Actual (%d)"), Manager->GetAudioTrackObserverFactory()->Count());
+			return true;
+		}
+
+		if (Manager->GetVideoTrackObserverFactory()->Count() != 1)
+		{
+			UE_LOG(LogPixelStreaming2RTC, Error, TEXT("VideoTrackObserverFactory has invalid reference count. Expected 1, Actual (%d)"), Manager->GetVideoTrackObserverFactory()->Count());
+			return true;
+		}
+
+		if (Manager->GetDataTrackObserverFactory()->Count() != 1)
+		{
+			UE_LOG(LogPixelStreaming2RTC, Error, TEXT("DataTrackObserverFactory has invalid reference count. Expected 1, Actual (%d)"), Manager->GetDataTrackObserverFactory()->Count());
+			return true;
+		}
+
+		return true;
+	}
+
+	bool FDisconnectSession::Update()
+	{
+		Manager->GetEpicRtcSession()->Disconnect(EpicRtcStringView{});
+
+		return true;
+	}
+
+	bool FCleanupSession::Update()
+	{
+		Manager->GetEpicRtcConference()->RemoveSession(ToEpicRtcStringView(SessionId));
+
+		// EpicRtc has released its hold on the session. All that should be holding a ref is the manager
+		if (Manager->GetEpicRtcSession()->Count() != 1)
+		{
+			UE_LOG(LogPixelStreaming2RTC, Error, TEXT("Session has invalid reference count. Expected 1, Actual (%d)"), Manager->GetEpicRtcSession()->Count());
+			return true;
+		}
+
+		// We know refcount was 1, so setting to nullptr should call the final release and destroy the session
+		Manager->GetEpicRtcSession() = nullptr;
+
+		// The session has been destroyed, the only thing holding a ref to the session observer should be the manager
+		if (Manager->GetSessionObserver()->Count() != 1)
+		{
+			UE_LOG(LogPixelStreaming2RTC, Error, TEXT("SessionObserver has invalid reference count. Expected 1, Actual (%d)"), Manager->GetSessionObserver()->Count());
+			return true;
+		}
+
+		return true;
+	}
+
+	bool FCleanupConference::Update()
+	{
+		Platform->ReleaseConference(ToEpicRtcStringView(ConferenceId));
+
+		return true;
+	}
+
+	bool FValidateRefCount::Update()
+	{
+		if (RefCountInterface->Count() != 1)
+		{
+			UE_LOG(LogPixelStreaming2RTC, Error, TEXT("Object has invalid reference count. Expected (%d), Actual (%d)"), 1, RefCountInterface->Count());
+			return true;
+		}
+
+		return true;
+	}
+
+	bool FCleanupManager::Update()
+	{
+		return true;
+	}
+} // namespace UE::PixelStreaming2
diff --git a/Engine/Plugins/Media/PixelStreaming2/Source/PixelStreaming2RTC/Private/Tests/MockPlayer.cpp b/Engine/Plugins/Media/PixelStreaming2/Source/PixelStreaming2RTC/Private/Tests/MockPlayer.cpp
new file mode 100644
--- /dev/null
+++ b/Engine/Plugins/Media/PixelStreaming2/Source/PixelStreaming2RTC/Private/Tests/MockPlayer.cpp
@@ -0,0 +1,298 @@
+// Copyright Epic Games, Inc. All Rights Reserved.
+
+#include "MockPlayer.h"
+
+#include "Async/Async.h"
+#include "DefaultDataProtocol.h"
+#include "EpicRtcVideoEncoderInitializer.h"
+#include "EpicRtcVideoDecoderInitializer.h"
+#include "EpicRtcWebsocketFactory.h"
+#include "Logging.h"
+#include "UtilsString.h"
+
+#include "epic_rtc/core/platform.h"
+
+#if WITH_DEV_AUTOMATION_TESTS
+
+namespace UE::PixelStreaming2
+{
+	uint32_t FMockPlayer::PlayerId = 0;
+
+	void FMockVideoSink::OnFrame(const EpicRtcVideoFrame& Frame)
+	{
+		VideoBuffer = Frame._buffer;
+		bReceivedFrame = true;
+	}
+
+	void FMockVideoSink::ResetReceivedFrame()
+	{
+		bReceivedFrame = false;
+		VideoBuffer.SafeRelease();
+	}
+
+	FMockPlayer::FMockPlayer()
+	{
+		PlayerName = FUtf8String::Printf("MockPlayer%d", PlayerId++);
+		FUtf8String ConferenceId = "test_conference";
+
+		EpicRtcVideoEncoderInitializers = { new FEpicRtcVideoEncoderInitializer() };
+		EpicRtcVideoDecoderInitializers = { new FEpicRtcVideoDecoderInitializer() };
+
+		EpicRtcManager = MakeShared<FEpicRtcManager>();
+
+		ToStreamerProtocol = UE::PixelStreaming2Input::GetDefaultToStreamerProtocol();
+
+		EpicRtcManager->OnSessionRoomsAvailableUpdate.AddRaw(this, &FMockPlayer::OnSessionRoomsAvailableUpdate);
+		EpicRtcManager->OnSessionErrorUpdate.AddRaw(this, &FMockPlayer::OnSessionErrorUpdate);
+		EpicRtcManager->OnRoomStateUpdate.AddRaw(this, &FMockPlayer::OnRoomStateUpdate);
+		EpicRtcManager->OnRoomJoinedUpdate.AddRaw(this, &FMockPlayer::OnRoomJoinedUpdate);
+		EpicRtcManager->OnRoomLeftUpdate.AddRaw(this, &FMockPlayer::OnRoomLeftUpdate);
+		EpicRtcManager->OnRoomErrorUpdate.AddRaw(this, &FMockPlayer::OnRoomErrorUpdate);
+		EpicRtcManager->OnSessionStateUpdate.AddRaw(this, &FMockPlayer::OnSessionStateUpdate);
+		EpicRtcManager->OnVideoTrackUpdate.AddRaw(this, &FMockPlayer::OnVideoTrackUpdate);
+		EpicRtcManager->OnVideoTrackFrame.AddRaw(this, &FMockPlayer::OnVideoTrackFrame);
+		EpicRtcManager->OnVideoTrackMuted.AddRaw(this, &FMockPlayer::OnVideoTrackMuted);
+		EpicRtcManager->OnVideoTrackRemoved.AddRaw(this, &FMockPlayer::OnVideoTrackRemoved);
+		EpicRtcManager->OnVideoTrackState.AddRaw(this, &FMockPlayer::OnVideoTrackState);
+
+		EpicRtcManager->OnDataTrackMessage.AddRaw(this, &FMockPlayer::OnDataTrackMessage);
+		EpicRtcManager->OnDataTrackRemoved.AddRaw(this, &FMockPlayer::OnDataTrackRemoved);
+		EpicRtcManager->OnDataTrackState.AddRaw(this, &FMockPlayer::OnDataTrackState);
+		EpicRtcManager->OnDataTrackUpdate.AddRaw(this, &FMockPlayer::OnDataTrackUpdate);
+
+		EpicRtcErrorCode Result = GetOrCreatePlatform({}, Platform.GetInitReference());
+
+		TRefCountPtr<FEpicRtcWebsocketFactory> WebsocketFactory = MakeRefCount<FEpicRtcWebsocketFactory>(false);
+
+		Result = Platform->CreateConference(ToEpicRtcStringView(ConferenceId),
+			{ ._websocketFactory = WebsocketFactory.GetReference(),
+				._signallingType = EpicRtcSignallingType::PixelStreaming,
+				._signingPlugin = nullptr,
+				._migrationPlugin = nullptr,
+				._audioDevicePlugin = nullptr,
+				._audioConfig = {
+					._tickAdm = true,
+					._audioEncoderInitializers = {},
+					._audioDecoderInitializers = {},
+					._enableBuiltInAudioCodecs = true,
+				},
+				._videoConfig = { ._videoEncoderInitializers = { ._ptr = const_cast<const EpicRtcVideoEncoderInitializerInterface**>(EpicRtcVideoEncoderInitializers.GetData()), ._size = (uint64_t)EpicRtcVideoEncoderInitializers.Num() }, ._videoDecoderInitializers = { ._ptr = const_cast<const EpicRtcVideoDecoderInitializerInterface**>(EpicRtcVideoDecoderInitializers.GetData()), ._size = (uint64_t)EpicRtcVideoDecoderInitializers.Num() }, ._enableBuiltInVideoCodecs = false },
+				._fieldTrials = { ._fieldTrials = EpicRtcStringView{ ._ptr = nullptr, ._length = 0 }, ._isGlobal = 0 } },
+			EpicRtcManager->EpicRtcConference.GetInitReference());
+
+		TickConferenceTask = FPixelStreamingTickableTask::Create<FEpicRtcTickConferenceTask>(EpicRtcManager->EpicRtcConference, TEXT("FMockPlayer TickConferenceTask"));
+
+		EpicRtcManager->SessionObserver = new FEpicRtcSessionObserver(EpicRtcManager);
+		EpicRtcManager->RoomObserver = new FEpicRtcRoomObserver(EpicRtcManager);
+
+		EpicRtcManager->AudioTrackObserverFactory = new FEpicRtcAudioTrackObserverFactory(EpicRtcManager);
+		EpicRtcManager->VideoTrackObserverFactory = new FEpicRtcVideoTrackObserverFactory(EpicRtcManager);
+		EpicRtcManager->DataTrackObserverFactory = new FEpicRtcDataTrackObserverFactory(EpicRtcManager);
+
+		VideoSink = MakeShared<FMockVideoSink>();
+	}
+
+	FMockPlayer::~FMockPlayer()
+	{
+		Disconnect(TEXT("Mock player being destroyed"));
+
+		if (EpicRtcManager->EpicRtcConference)
+		{
+			EpicRtcManager->EpicRtcConference->RemoveSession(ToEpicRtcStringView(PlayerName));
+
+			Platform->ReleaseConference(EpicRtcManager->EpicRtcConference->GetId());
+		}
+	}
+
+	void FMockPlayer::Connect(int StreamerPort)
+	{
+		FUtf8String Url(FString::Printf(TEXT("ws://127.0.0.1:%d/"), StreamerPort));
+		FUtf8String ConnectionUrl = Url + +(Url.Contains(TEXT("?")) ? TEXT("&") : TEXT("?")) + TEXT("isStreamer=false");
+
+		EpicRtcSessionConfig SessionConfig = {
+			._id = ToEpicRtcStringView(PlayerName),
+			._url = ToEpicRtcStringView(ConnectionUrl),
+			._observer = EpicRtcManager->SessionObserver.GetReference()
+		};
+
+		EpicRtcErrorCode Result = EpicRtcManager->EpicRtcConference->CreateSession(SessionConfig, EpicRtcManager->EpicRtcSession.GetInitReference());
+		if (Result != EpicRtcErrorCode::Ok)
+		{
+			UE_LOG(LogPixelStreaming2RTC, Error, TEXT("FMockPlayer Failed to create EpicRtc session"));
+			return;
+		}
+
+		Result = EpicRtcManager->EpicRtcSession->Connect();
+		if (Result != EpicRtcErrorCode::Ok)
+		{
+			UE_LOG(LogPixelStreaming2RTC, Error, TEXT("FMockPlayer Failed to connect EpicRtcSession"));
+		}
+		else
+		{
+			UE_LOG(LogPixelStreaming2RTC, VeryVerbose, TEXT("FMockPlayer Connected to EpicRtcSession"));
+		}
+	}
+
+	bool FMockPlayer::Subscribe(const FString& StreamerId)
+	{
+		if (SessionState != EpicRtcSessionState::Connected)
+		{
+			// Sessions state can take several ticks to returning false tells latent test to run again next tick.
+			return false;
+		}
+
+		EpicRtcConnectionConfig ConnectionConfig = {
+			._iceServers = EpicRtcIceServerSpan{ ._ptr = nullptr, ._size = 0 },
+			._iceConnectionPolicy = EpicRtcIcePolicy::All,
+			._disableTcpCandidates = false
+		};
+
+		SubscribedStream = *StreamerId;
+
+		EpicRtcRoomConfig RoomConfig = {
+			._id = ToEpicRtcStringView(SubscribedStream),
+			._connectionConfig = ConnectionConfig,
+			._ticket = EpicRtcStringView{ ._ptr = nullptr, ._length = 0 },
+			._observer = EpicRtcManager->RoomObserver,
+			._audioTrackObserverFactory = EpicRtcManager->AudioTrackObserverFactory,
+			._dataTrackObserverFactory = EpicRtcManager->DataTrackObserverFactory,
+			._videoTrackObserverFactory = EpicRtcManager->VideoTrackObserverFactory
+		};
+
+		EpicRtcErrorCode Result = EpicRtcManager->EpicRtcSession->CreateRoom(RoomConfig, EpicRtcManager->EpicRtcRoom.GetInitReference());
+		if (Result != EpicRtcErrorCode::Ok)
+		{
+			UE_LOG(LogPixelStreaming2RTC, Error, TEXT("Failed to create EpicRtc room"));
+			return false;
+		}
+
+		EpicRtcManager->EpicRtcRoom->Join();
+
+		return true;
+	}
+
+	void FMockPlayer::OnVideoTrackUpdate(EpicRtcParticipantInterface* Participant, EpicRtcVideoTrackInterface* VideoTrack)
+	{
+		FString ParticipantId{ (int32)Participant->GetId()._length, Participant->GetId()._ptr };
+		FString VideoTrackId{ (int32)VideoTrack->GetId()._length, VideoTrack->GetId()._ptr };
+		UE_LOG(LogPixelStreaming2RTC, VeryVerbose, TEXT("UPixelStreaming2Peer::OnVideoTrackUpdate(Participant [%s], VideoTrack [%s])"), *ParticipantId, *VideoTrackId);
+	}
+
+	void FMockPlayer::OnVideoTrackFrame(EpicRtcVideoTrackInterface* VideoTrack, const EpicRtcVideoFrame& Frame)
+	{
+		UE_LOG(LogPixelStreaming2RTC, VeryVerbose, TEXT("FMockPlayer::OnVideoTrackFrame received a video frame."));
+
+		VideoSink->OnFrame(Frame);
+	}
+
+	void FMockPlayer::OnVideoTrackMuted(EpicRtcVideoTrackInterface* VideoTrack, EpicRtcBool bIsMuted)
+	{
+	}
+
+	void FMockPlayer::OnVideoTrackRemoved(EpicRtcVideoTrackInterface* VideoTrack)
+	{
+	}
+
+	void FMockPlayer::OnVideoTrackState(EpicRtcVideoTrackInterface* VideoTrack, const EpicRtcTrackState State)
+	{
+	}
+
+	void FMockPlayer::OnSessionRoomsAvailableUpdate(EpicRtcStringArrayInterface* RoomsList)
+	{
+	}
+
+	void FMockPlayer::OnSessionErrorUpdate(const EpicRtcErrorCode ErrorUpdate)
+	{
+		UE_LOG(LogPixelStreaming2RTC, Log, TEXT("OnSessionErrorUpdate: "));
+	}
+
+	void FMockPlayer::OnRoomStateUpdate(const EpicRtcRoomState State)
+	{
+		UE_LOG(LogPixelStreaming2RTC, Log, TEXT("OnRoomStateUpdate: "));
+	}
+
+	void FMockPlayer::OnRoomJoinedUpdate(EpicRtcParticipantInterface* Participant)
+	{
+		FString ParticipantId{ (int32)Participant->GetId()._length, Participant->GetId()._ptr };
+		UE_LOG(LogPixelStreaming2RTC, Log, TEXT("OnRoomJoinedUpdate: Player (%s) joined"), *ParticipantId);
+	}
+
+	void FMockPlayer::OnRoomLeftUpdate(const EpicRtcStringView ParticipantId)
+	{
+		UE_LOG(LogPixelStreaming2RTC, Log, TEXT("OnRoomLeftUpdate"));
+	}
+
+	void FMockPlayer::OnRoomErrorUpdate(const EpicRtcErrorCode Error)
+	{
+		UE_LOG(LogPixelStreaming2RTC, Log, TEXT("OnRoomErrorUpdate"));
+	}
+
+	void FMockPlayer::OnSessionStateUpdate(const EpicRtcSessionState StateUpdate)
+	{
+		switch (StateUpdate)
+		{
+			case EpicRtcSessionState::New:
+			case EpicRtcSessionState::Pending:
+			case EpicRtcSessionState::Connected:
+			case EpicRtcSessionState::Disconnected:
+			case EpicRtcSessionState::Failed:
+			case EpicRtcSessionState::Exiting:
+				SessionState = StateUpdate;
+				break;
+			default:
+				break;
+		}
+	}
+
+	void FMockPlayer::OnDataTrackMessage(EpicRtcDataTrackInterface* InDataTrack)
+	{
+		TRefCountPtr<EpicRtcDataFrameInterface> DataFrame;
+		if (!InDataTrack->PopFrame(DataFrame.GetInitReference()))
+		{
+			UE_LOG(LogPixelStreaming2RTC, Error, TEXT("FMockPlayer::OnDataTrackMessage Failed to PopFrame"));
+			return;
+		}
+
+		// Broadcast must be done on the GameThread because the GameThread can remove the delegates.
+		// If removing and broadcast happens simultaneously it causes a datarace failure.
+		TWeakPtr<FMockPlayer> WeakPlayer = AsShared();
+		AsyncTask(ENamedThreads::GameThread, [WeakPlayer, DataFrame]() {
+			if (TSharedPtr<FMockPlayer> PinnedPlayer = WeakPlayer.Pin())
+			{
+				const TArray<uint8> Data(DataFrame->Data(), DataFrame->Size());
+				PinnedPlayer->OnMessageReceived.Broadcast(Data);
+			}
+		});
+	}
+
+	void FMockPlayer::OnDataTrackRemoved(EpicRtcDataTrackInterface*) {}
+
+	void FMockPlayer::OnDataTrackState(EpicRtcDataTrackInterface*, const EpicRtcTrackState) {}
+
+	void FMockPlayer::OnDataTrackUpdate(EpicRtcParticipantInterface*, EpicRtcDataTrackInterface* InDataTrack)
+	{
+		DataTrack = FEpicRtcDataTrack::Create(InDataTrack, ToStreamerProtocol);
+	}
+
+	void FMockPlayer::Disconnect(const FString& Reason)
+	{
+		if (!EpicRtcManager->EpicRtcSession)
+		{
+			return;
+		}
+
+		if (EpicRtcManager->EpicRtcRoom)
+		{
+			EpicRtcManager->EpicRtcRoom->Leave();
+			EpicRtcManager->EpicRtcSession->RemoveRoom(ToEpicRtcStringView(SubscribedStream));
+		}
+
+		EpicRtcErrorCode Result = EpicRtcManager->EpicRtcSession->Disconnect(ToEpicRtcStringView(*Reason));
+		if (Result != EpicRtcErrorCode::Ok)
+		{
+			UE_LOG(LogPixelStreaming2RTC, Error, TEXT("Failed to disconnect EpicRtcSession"));
+		}
+	}
+} // namespace UE::PixelStreaming2
+
+#endif // WITH_DEV_AUTOMATION_TESTS
diff --git a/Engine/Plugins/Media/PixelStreaming2/Source/PixelStreaming2RTC/Private/Tests/MockPlayer.h b/Engine/Plugins/Media/PixelStreaming2/Source/PixelStreaming2RTC/Private/Tests/MockPlayer.h
new file mode 100644
--- /dev/null
+++ b/Engine/Plugins/Media/PixelStreaming2/Source/PixelStreaming2RTC/Private/Tests/MockPlayer.h
@@ -0,0 +1,107 @@
+// Copyright Epic Games, Inc. All Rights Reserved.
+
+#pragma once
+
+#include "EpicRtcDataTrack.h"
+#include "HAL/ThreadSafeBool.h"
+#include "EpicRtcManager.h"
+#include "EpicRtcConferenceUtils.h"
+#include "IPixelStreaming2DataProtocol.h"
+#include "Containers/Utf8String.h"
+
+#if WITH_DEV_AUTOMATION_TESTS
+
+class EpicRtcPlatformInterface;
+
+namespace UE::PixelStreaming2
+{
+	class FEpicRtcDataTrack;
+
+	struct FMockVideoFrameConfig
+	{
+		int	  Height;
+		int	  Width;
+		uint8 Y;
+		uint8 U;
+		uint8 V;
+	};
+
+	class FMockVideoSink
+	{
+	public:
+		void									  OnFrame(const EpicRtcVideoFrame& Frame);
+		bool									  HasReceivedFrame() const { return bReceivedFrame; }
+		void									  ResetReceivedFrame();
+		TRefCountPtr<EpicRtcVideoBufferInterface> GetReceivedBuffer() { return VideoBuffer; };
+
+	private:
+		TRefCountPtr<EpicRtcVideoBufferInterface> VideoBuffer;
+		FThreadSafeBool							  bReceivedFrame = false;
+	};
+
+	class FMockPlayer : public TSharedFromThis<FMockPlayer>
+	{
+	public:
+		FMockPlayer();
+		virtual ~FMockPlayer();
+
+		void Connect(int StreamerPort);
+		void Disconnect(const FString& Reason);
+		bool Subscribe(const FString& StreamerId);
+
+		void OnVideoTrackUpdate(EpicRtcParticipantInterface* Participant, EpicRtcVideoTrackInterface* VideoTrack);
+		void OnVideoTrackFrame(EpicRtcVideoTrackInterface* VideoTrack, const EpicRtcVideoFrame& Frame);
+		void OnVideoTrackMuted(EpicRtcVideoTrackInterface* VideoTrack, EpicRtcBool bIsMuted);
+		void OnVideoTrackRemoved(EpicRtcVideoTrackInterface* VideoTrack);
+		void OnVideoTrackState(EpicRtcVideoTrackInterface* VideoTrack, const EpicRtcTrackState State);
+		void OnSessionRoomsAvailableUpdate(EpicRtcStringArrayInterface* RoomsList);
+		void OnSessionErrorUpdate(const EpicRtcErrorCode ErrorUpdate);
+		void OnRoomStateUpdate(const EpicRtcRoomState State);
+		void OnRoomJoinedUpdate(EpicRtcParticipantInterface* Participant);
+		void OnRoomLeftUpdate(const EpicRtcStringView ParticipantId);
+		void OnRoomErrorUpdate(const EpicRtcErrorCode Error);
+		void OnSessionStateUpdate(const EpicRtcSessionState StateUpdate);
+		void OnDataTrackMessage(EpicRtcDataTrackInterface*);
+		void OnDataTrackRemoved(EpicRtcDataTrackInterface*);
+		void OnDataTrackState(EpicRtcDataTrackInterface*, const EpicRtcTrackState);
+		void OnDataTrackUpdate(EpicRtcParticipantInterface*, EpicRtcDataTrackInterface*);
+
+		template <typename... Args>
+		bool SendMessage(FString MessageType, Args... VarArgs)
+		{
+			if (!DataTrack)
+			{
+				return false;
+			}
+
+			return DataTrack->SendMessage(MessageType, VarArgs...);
+		}
+
+		bool DataChannelAvailable() { return DataTrack.IsValid(); };
+
+		TSharedPtr<FMockVideoSink>					   GetVideoSink() { return VideoSink; };
+		TSharedPtr<IPixelStreaming2DataProtocol> GetToStreamerProtocol() { return ToStreamerProtocol; }
+
+		DECLARE_MULTICAST_DELEGATE_OneParam(FOnMessageReceived, const TArray<uint8>&);
+		FOnMessageReceived OnMessageReceived;
+
+	private:
+		TSharedPtr<FEpicRtcManager>					   EpicRtcManager;
+		TSharedPtr<FMockVideoSink>					   VideoSink;
+		TSharedPtr<FEpicRtcDataTrack>				   DataTrack;
+		TRefCountPtr<EpicRtcPlatformInterface>		   Platform;
+		TUniqueTaskPtr<FEpicRtcTickConferenceTask>	   TickConferenceTask;
+		TSharedPtr<IPixelStreaming2DataProtocol> ToStreamerProtocol;
+
+		TArray<EpicRtcVideoEncoderInitializerInterface*> EpicRtcVideoEncoderInitializers;
+		TArray<EpicRtcVideoDecoderInitializerInterface*> EpicRtcVideoDecoderInitializers;
+
+		EpicRtcSessionState SessionState = EpicRtcSessionState::Disconnected;
+
+		FUtf8String		SubscribedStream;
+		FUtf8String		PlayerName;
+		static uint32_t PlayerId;
+	};
+} // namespace UE::PixelStreaming2
+
+#endif // WITH_DEV_AUTOMATION_TESTS
diff --git a/Engine/Plugins/Media/PixelStreaming2/Source/PixelStreaming2RTC/Private/Tests/TestUtils.h b/Engine/Plugins/Media/PixelStreaming2/Source/PixelStreaming2RTC/Private/Tests/TestUtils.h
new file mode 100644
--- /dev/null
+++ b/Engine/Plugins/Media/PixelStreaming2/Source/PixelStreaming2RTC/Private/Tests/TestUtils.h
@@ -0,0 +1,89 @@
+// Copyright Epic Games, Inc. All Rights Reserved.
+
+#pragma once
+
+#include "IPixelStreaming2Streamer.h"
+#include "Misc/AutomationTest.h"
+#include "MockPlayer.h"
+#include "PixelStreaming2Servers.h"
+#include "VideoProducer.h"
+#include "Video/VideoConfig.h"
+
+#if WITH_DEV_AUTOMATION_TESTS
+
+namespace UE::PixelStreaming2
+{
+	namespace TestUtils
+	{
+		int32 NextStreamerPort();
+		int32 NextPlayerPort();
+	} // namespace TestUtils
+
+	// Equivalent to DEFINE_LATENT_AUTOMATION_COMMAND_FOUR_PARAMETER, but instead we define a custom constructor
+	class FWaitForDataChannelMessageOrTimeout : public IAutomationLatentCommand
+	{
+	public:
+		FWaitForDataChannelMessageOrTimeout(double InTimeoutSeconds, TSharedPtr<FMockPlayer> InPlayer, TFunction<void(const TArray<uint8>&)>& InCallback, TSharedPtr<bool> InbComplete)
+			: TimeoutSeconds(InTimeoutSeconds)
+			, Player(InPlayer)
+			, Callback(InCallback)
+			, bComplete(InbComplete)
+		{
+			MessageReceivedHandle = Player->OnMessageReceived.AddLambda([MessageCallback = this->Callback](const TArray<uint8>& RawBuffer) {
+				MessageCallback(RawBuffer);
+			});
+		}
+
+		virtual ~FWaitForDataChannelMessageOrTimeout()
+		{
+			Player->OnMessageReceived.Remove(MessageReceivedHandle);
+		}
+		virtual bool Update() override;
+
+	private:
+		double								  TimeoutSeconds;
+		TSharedPtr<FMockPlayer>				  Player;
+		TFunction<void(const TArray<uint8>&)> Callback;
+		TSharedPtr<bool>					  bComplete;
+		FDelegateHandle						  MessageReceivedHandle;
+	};
+
+	class FWaitForStreamerDataChannelMessageOrTimeout : public IAutomationLatentCommand
+	{
+	public:
+		FWaitForStreamerDataChannelMessageOrTimeout(double InTimeoutSeconds, TSharedPtr<IPixelStreaming2Streamer> InStreamer, TSharedPtr<bool> InbComplete)
+			: TimeoutSeconds(InTimeoutSeconds)
+			, Streamer(InStreamer)
+			, bComplete(InbComplete)
+		{
+		}
+
+		virtual ~FWaitForStreamerDataChannelMessageOrTimeout()
+		{
+		}
+		virtual bool Update() override;
+
+	private:
+		double									   TimeoutSeconds;
+		TSharedPtr<IPixelStreaming2Streamer> Streamer;
+		TSharedPtr<bool>						   bComplete;
+	};
+
+	DEFINE_LATENT_AUTOMATION_COMMAND_ONE_PARAMETER(FWaitSeconds, double, WaitSeconds);
+	DEFINE_LATENT_AUTOMATION_COMMAND_TWO_PARAMETER(FSendSolidColorFrame, TSharedPtr<FVideoProducer>, VideoProducer, FMockVideoFrameConfig, FrameConfig);
+	DEFINE_LATENT_AUTOMATION_COMMAND_THREE_PARAMETER(FWaitForFrameReceived, double, TimeoutSeconds, TSharedPtr<FMockVideoSink>, VideoSink, FMockVideoFrameConfig, FrameConfig);
+	DEFINE_LATENT_AUTOMATION_COMMAND_FOUR_PARAMETER(FSubscribePlayerAfterStreamerConnectedOrTimeout, double, TimeoutSeconds, TSharedPtr<IPixelStreaming2Streamer>, OutStreamer, TSharedPtr<FMockPlayer>, OutPlayer, FString, StreamerName);
+	DEFINE_LATENT_AUTOMATION_COMMAND_THREE_PARAMETER(FCleanupAll, TSharedPtr<UE::PixelStreaming2Servers::IServer>, OutSignallingServer, TSharedPtr<IPixelStreaming2Streamer>, OutStreamer, TSharedPtr<FMockPlayer>, OutPlayer);
+	DEFINE_LATENT_AUTOMATION_COMMAND_TWO_PARAMETER(FWaitForDataChannelOrTimeout, double, TimeoutSeconds, TSharedPtr<FMockPlayer>, OutPlayer);
+	DEFINE_LATENT_AUTOMATION_COMMAND_THREE_PARAMETER(FSendDataChannelMessageToStreamer, TSharedPtr<FMockPlayer>, Player, FString, MessageType, const FString, Body);
+	DEFINE_LATENT_AUTOMATION_COMMAND_THREE_PARAMETER(FSendDataChannelMessageFromStreamer, TSharedPtr<IPixelStreaming2Streamer>, Streamer, FString, MessageType, const FString, Body);
+	DEFINE_LATENT_AUTOMATION_COMMAND_THREE_PARAMETER(FSendCustomMessageToStreamer, TSharedPtr<FMockPlayer>, Player, FString, MessageType, uint16, Body);
+
+	TSharedPtr<IPixelStreaming2Streamer>			  CreateStreamer(const FString& StreamerName, int StreamerPort);
+	TSharedPtr<FMockPlayer>								  CreatePlayer();
+	TSharedPtr<UE::PixelStreaming2Servers::IServer> CreateSignallingServer(int StreamerPort, int PlayerPort);
+	void												  SetCodec(EVideoCodec Codec);
+	void												  SetDisableTransmitVideo(bool bDisableTransmitVideo);
+} // namespace UE::PixelStreaming2
+
+#endif // WITH_DEV_AUTOMATION_TESTS
\ No newline at end of file
diff --git a/Engine/Plugins/Media/PixelStreaming2/Source/PixelStreaming2RTC/Private/UtilsCodecs.h b/Engine/Plugins/Media/PixelStreaming2/Source/PixelStreaming2RTC/Private/UtilsCodecs.h
new file mode 100644
--- /dev/null
+++ b/Engine/Plugins/Media/PixelStreaming2/Source/PixelStreaming2RTC/Private/UtilsCodecs.h
@@ -0,0 +1,181 @@
+// Copyright Epic Games, Inc. All Rights Reserved.
+
+#pragma once
+
+#include "Containers/Array.h"
+#include "EpicRtcVideoCommon.h"
+#include "IPixelStreaming2Streamer.h"
+#include "PixelStreaming2PluginSettings.h"
+#include "Video/VideoEncoder.h"
+#include "UtilsVideo.h"
+
+#include "epic_rtc/core/video/video_common.h"
+#include "epic_rtc/core/video/video_rate_control.h"
+
+namespace UE::PixelStreaming2
+{
+	// Helper array for all scalability modes. EScalabilityMode::None must always be the last entry
+	const TArray<EScalabilityMode> AllScalabilityModes = {
+		EScalabilityMode::L1T1,
+		EScalabilityMode::L1T2,
+		EScalabilityMode::L1T3,
+		EScalabilityMode::L2T1,
+		EScalabilityMode::L2T1h,
+		EScalabilityMode::L2T1_KEY,
+		EScalabilityMode::L2T2,
+		EScalabilityMode::L2T2h,
+		EScalabilityMode::L2T2_KEY,
+		EScalabilityMode::L2T2_KEY_SHIFT,
+		EScalabilityMode::L2T3,
+		EScalabilityMode::L2T3h,
+		EScalabilityMode::L2T3_KEY,
+		EScalabilityMode::L3T1,
+		EScalabilityMode::L3T1h,
+		EScalabilityMode::L3T1_KEY,
+		EScalabilityMode::L3T2,
+		EScalabilityMode::L3T2h,
+		EScalabilityMode::L3T2_KEY,
+		EScalabilityMode::L3T3,
+		EScalabilityMode::L3T3h,
+		EScalabilityMode::L3T3_KEY,
+		EScalabilityMode::S2T1,
+		EScalabilityMode::S2T1h,
+		EScalabilityMode::S2T2,
+		EScalabilityMode::S2T2h,
+		EScalabilityMode::S2T3,
+		EScalabilityMode::S2T3h,
+		EScalabilityMode::S3T1,
+		EScalabilityMode::S3T1h,
+		EScalabilityMode::S3T2,
+		EScalabilityMode::S3T2h,
+		EScalabilityMode::S3T3,
+		EScalabilityMode::S3T3h,
+		EScalabilityMode::None
+	};
+
+	// Make sure EpicRtcVideoScalabilityMode and EScalabilityMode match up
+	static_assert(EpicRtcVideoScalabilityMode::L1T1 == static_cast<EpicRtcVideoScalabilityMode>(EScalabilityMode::L1T1));
+	static_assert(EpicRtcVideoScalabilityMode::L1T2 == static_cast<EpicRtcVideoScalabilityMode>(EScalabilityMode::L1T2));
+	static_assert(EpicRtcVideoScalabilityMode::L1T3 == static_cast<EpicRtcVideoScalabilityMode>(EScalabilityMode::L1T3));
+	static_assert(EpicRtcVideoScalabilityMode::L2T1 == static_cast<EpicRtcVideoScalabilityMode>(EScalabilityMode::L2T1));
+	static_assert(EpicRtcVideoScalabilityMode::L2T1h == static_cast<EpicRtcVideoScalabilityMode>(EScalabilityMode::L2T1h));
+	static_assert(EpicRtcVideoScalabilityMode::L2T1Key == static_cast<EpicRtcVideoScalabilityMode>(EScalabilityMode::L2T1_KEY));
+	static_assert(EpicRtcVideoScalabilityMode::L2T2 == static_cast<EpicRtcVideoScalabilityMode>(EScalabilityMode::L2T2));
+	static_assert(EpicRtcVideoScalabilityMode::L2T2h == static_cast<EpicRtcVideoScalabilityMode>(EScalabilityMode::L2T2h));
+	static_assert(EpicRtcVideoScalabilityMode::L2T2Key == static_cast<EpicRtcVideoScalabilityMode>(EScalabilityMode::L2T2_KEY));
+	static_assert(EpicRtcVideoScalabilityMode::L2T2KeyShift == static_cast<EpicRtcVideoScalabilityMode>(EScalabilityMode::L2T2_KEY_SHIFT));
+	static_assert(EpicRtcVideoScalabilityMode::L2T3 == static_cast<EpicRtcVideoScalabilityMode>(EScalabilityMode::L2T3));
+	static_assert(EpicRtcVideoScalabilityMode::L2T3h == static_cast<EpicRtcVideoScalabilityMode>(EScalabilityMode::L2T3h));
+	static_assert(EpicRtcVideoScalabilityMode::L2T3Key == static_cast<EpicRtcVideoScalabilityMode>(EScalabilityMode::L2T3_KEY));
+	static_assert(EpicRtcVideoScalabilityMode::L3T1 == static_cast<EpicRtcVideoScalabilityMode>(EScalabilityMode::L3T1));
+	static_assert(EpicRtcVideoScalabilityMode::L3T1h == static_cast<EpicRtcVideoScalabilityMode>(EScalabilityMode::L3T1h));
+	static_assert(EpicRtcVideoScalabilityMode::L3T1Key == static_cast<EpicRtcVideoScalabilityMode>(EScalabilityMode::L3T1_KEY));
+	static_assert(EpicRtcVideoScalabilityMode::L3T2 == static_cast<EpicRtcVideoScalabilityMode>(EScalabilityMode::L3T2));
+	static_assert(EpicRtcVideoScalabilityMode::L3T2h == static_cast<EpicRtcVideoScalabilityMode>(EScalabilityMode::L3T2h));
+	static_assert(EpicRtcVideoScalabilityMode::L3T2Key == static_cast<EpicRtcVideoScalabilityMode>(EScalabilityMode::L3T2_KEY));
+	static_assert(EpicRtcVideoScalabilityMode::L3T3 == static_cast<EpicRtcVideoScalabilityMode>(EScalabilityMode::L3T3));
+	static_assert(EpicRtcVideoScalabilityMode::L3T3h == static_cast<EpicRtcVideoScalabilityMode>(EScalabilityMode::L3T3h));
+	static_assert(EpicRtcVideoScalabilityMode::L3T3Key == static_cast<EpicRtcVideoScalabilityMode>(EScalabilityMode::L3T3_KEY));
+	static_assert(EpicRtcVideoScalabilityMode::S2T1 == static_cast<EpicRtcVideoScalabilityMode>(EScalabilityMode::S2T1));
+	static_assert(EpicRtcVideoScalabilityMode::S2T1h == static_cast<EpicRtcVideoScalabilityMode>(EScalabilityMode::S2T1h));
+	static_assert(EpicRtcVideoScalabilityMode::S2T2 == static_cast<EpicRtcVideoScalabilityMode>(EScalabilityMode::S2T2));
+	static_assert(EpicRtcVideoScalabilityMode::S2T2h == static_cast<EpicRtcVideoScalabilityMode>(EScalabilityMode::S2T2h));
+	static_assert(EpicRtcVideoScalabilityMode::S2T3 == static_cast<EpicRtcVideoScalabilityMode>(EScalabilityMode::S2T3));
+	static_assert(EpicRtcVideoScalabilityMode::S2T3h == static_cast<EpicRtcVideoScalabilityMode>(EScalabilityMode::S2T3h));
+	static_assert(EpicRtcVideoScalabilityMode::S3T1 == static_cast<EpicRtcVideoScalabilityMode>(EScalabilityMode::S3T1));
+	static_assert(EpicRtcVideoScalabilityMode::S3T1h == static_cast<EpicRtcVideoScalabilityMode>(EScalabilityMode::S3T1h));
+	static_assert(EpicRtcVideoScalabilityMode::S3T2 == static_cast<EpicRtcVideoScalabilityMode>(EScalabilityMode::S3T2));
+	static_assert(EpicRtcVideoScalabilityMode::S3T2h == static_cast<EpicRtcVideoScalabilityMode>(EScalabilityMode::S3T2h));
+	static_assert(EpicRtcVideoScalabilityMode::S3T3 == static_cast<EpicRtcVideoScalabilityMode>(EScalabilityMode::S3T3));
+	static_assert(EpicRtcVideoScalabilityMode::S3T3h == static_cast<EpicRtcVideoScalabilityMode>(EScalabilityMode::S3T3h));
+	static_assert(EpicRtcVideoScalabilityMode::None == static_cast<EpicRtcVideoScalabilityMode>(EScalabilityMode::None));
+
+	inline FEpicRtcParameterPairArray* CreateH264Format(EH264Profile Profile, UE::AVCodecCore::H264::EH264Level Level)
+	{
+		// TODO (Migration): RTCP-7028 picRtcStringView needs a way to own the memory passed into it
+		// return new FEpicRtcParameterPairArray(
+		// {
+		// 	EpicRtcParameterPair{
+		// 		._key = EpicRtcStringView{ ._ptr = "profile-level-id", ._length = 16 },
+		// 		._value = EpicRtcStringView{ ._ptr = TCHAR_TO_ANSI(*ProfileString), ._length = (uint64_t)ProfileString.Len() }
+		// 	},
+		// 	EpicRtcParameterPair{
+		// 		._key = EpicRtcStringView{ ._ptr = "packetization-mode", ._length = 18 },
+		// 		._value = EpicRtcStringView{ ._ptr = "1", ._length = 1 }
+		// 	},
+		// 	EpicRtcParameterPair{
+		// 		._key = EpicRtcStringView{ ._ptr = "level-asymmetry-allowed", ._length = 23 },
+		// 		._value = EpicRtcStringView{ ._ptr = "1", ._length = 1 }
+		// 	}
+		// });
+
+		using namespace UE::AVCodecCore::H264;
+		if (Profile == EH264Profile::ConstrainedBaseline && Level == EH264Level::Level_3_1)
+		{
+			return new FEpicRtcParameterPairArray(
+				{ EpicRtcParameterPair{
+					  ._key = EpicRtcStringView{ ._ptr = "profile-level-id", ._length = 16 },
+					  ._value = EpicRtcStringView{ ._ptr = "42e01f", ._length = 6 } },
+					EpicRtcParameterPair{
+						._key = EpicRtcStringView{ ._ptr = "packetization-mode", ._length = 18 },
+						._value = EpicRtcStringView{ ._ptr = "1", ._length = 1 } },
+					EpicRtcParameterPair{
+						._key = EpicRtcStringView{ ._ptr = "level-asymmetry-allowed", ._length = 23 },
+						._value = EpicRtcStringView{ ._ptr = "1", ._length = 1 } } });
+		}
+		else if (Profile == EH264Profile::Baseline && Level == EH264Level::Level_3_1)
+		{
+			return new FEpicRtcParameterPairArray(
+				{ EpicRtcParameterPair{
+					  ._key = EpicRtcStringView{ ._ptr = "profile-level-id", ._length = 16 },
+					  ._value = EpicRtcStringView{ ._ptr = "42001f", ._length = 6 } },
+					EpicRtcParameterPair{
+						._key = EpicRtcStringView{ ._ptr = "packetization-mode", ._length = 18 },
+						._value = EpicRtcStringView{ ._ptr = "1", ._length = 1 } },
+					EpicRtcParameterPair{
+						._key = EpicRtcStringView{ ._ptr = "level-asymmetry-allowed", ._length = 23 },
+						._value = EpicRtcStringView{ ._ptr = "1", ._length = 1 } } });
+		}
+		else
+		{
+			return nullptr;
+		}
+	}
+
+	inline FEpicRtcParameterPairArray* CreateVP9Format(UE::AVCodecCore::VP9::EProfile Profile)
+	{
+		using namespace UE::AVCodecCore::VP9;
+		if (Profile == EProfile::Profile0)
+		{
+			return new FEpicRtcParameterPairArray(
+				{ EpicRtcParameterPair{
+					._key = EpicRtcStringView{ ._ptr = "profile-id", ._length = 10 },
+					._value = EpicRtcStringView{ ._ptr = "0", ._length = 1 } } });
+		}
+		else if (Profile == EProfile::Profile1)
+		{
+			return new FEpicRtcParameterPairArray(
+				{ EpicRtcParameterPair{
+					._key = EpicRtcStringView{ ._ptr = "profile-id", ._length = 10 },
+					._value = EpicRtcStringView{ ._ptr = "1", ._length = 1 } } });
+		}
+		else if (Profile == EProfile::Profile2)
+		{
+			return new FEpicRtcParameterPairArray(
+				{ EpicRtcParameterPair{
+					._key = EpicRtcStringView{ ._ptr = "profile-id", ._length = 10 },
+					._value = EpicRtcStringView{ ._ptr = "2", ._length = 1 } } });
+		}
+		else if (Profile == EProfile::Profile3)
+		{
+			return new FEpicRtcParameterPairArray(
+				{ EpicRtcParameterPair{
+					._key = EpicRtcStringView{ ._ptr = "profile-id", ._length = 10 },
+					._value = EpicRtcStringView{ ._ptr = "3", ._length = 1 } } });
+		}
+		else
+		{
+			return nullptr;
+		}
+	}
+} // namespace UE::PixelStreaming2
\ No newline at end of file
