[PS2] PR #12718: Refactor PixelStreaming2 plugin to enable other streaming protocols

diff --git a/Engine/Plugins/Media/PixelStreaming2/Source/PixelStreaming2Settings/Private/PixelStreaming2PluginSettings.cpp b/Engine/Plugins/Media/PixelStreaming2/Source/PixelStreaming2Settings/Private/PixelStreaming2PluginSettings.cpp
new file mode 100644
--- /dev/null
+++ b/Engine/Plugins/Media/PixelStreaming2/Source/PixelStreaming2Settings/Private/PixelStreaming2PluginSettings.cpp
@@ -0,0 +1,1303 @@
+// Copyright Epic Games, Inc. All Rights Reserved.
+
+#include "PixelStreaming2PluginSettings.h"
+#include "Logging.h"
+#include "Misc/CommandLine.h"
+#include "UObject/ReflectedTypeAccessors.h"
+
+namespace
+{
+	template <typename TEnumType>
+	static void CheckConsoleEnum(IConsoleVariable* InConsoleVariable)
+	{
+		FString ConsoleString = InConsoleVariable->GetString();
+		if (StaticEnum<TEnumType>()->GetIndexByNameString(ConsoleString) == INDEX_NONE)
+		{
+			// Legacy CVar values were the enum values but underscores (LOW_LATENCY) instead of the camel case UENUM string (LowLatency). They are still valid we just need to remove the underscores when we check them.
+			if (ConsoleString = ConsoleString.Replace(TEXT("_"), TEXT("")); StaticEnum<TEnumType>()->GetIndexByNameString(ConsoleString) != INDEX_NONE)
+			{
+				InConsoleVariable->Set(*ConsoleString, ECVF_SetByConsole);
+			}
+			else
+			{
+				FString ConsoleObjectName = IConsoleManager::Get().FindConsoleObjectName(InConsoleVariable);
+				UE_LOGFMT(LogPixelStreaming2Settings, Warning, "Invalid value {0} received for enum {1} of type {2}", ConsoleString, ConsoleObjectName, StaticEnum<TEnumType>()->GetName());
+				InConsoleVariable->Set(*InConsoleVariable->GetDefaultValue(), ECVF_SetByConsole);
+			}
+		}
+	}
+
+	static void VerifyCVarVideoSettings(IConsoleVariable* /* We ignore the passed in console variable as this method is called by many different CVars */)
+	{
+		IConsoleVariable* SimulcastCVar = IConsoleManager::Get().FindConsoleVariable(TEXT("PixelStreaming2.Encoder.EnableSimulcast"));
+		IConsoleVariable* CodecCVar = IConsoleManager::Get().FindConsoleVariable(TEXT("PixelStreaming2.Encoder.Codec"));
+		IConsoleVariable* ScalabilityModeCVar = IConsoleManager::Get().FindConsoleVariable(TEXT("PixelStreaming2.Encoder.ScalabilityMode"));
+
+		// Verify that the video codec and scalability mode strings correctly map to an enum
+		CheckConsoleEnum<EVideoCodec>(CodecCVar);
+		CheckConsoleEnum<EScalabilityMode>(ScalabilityModeCVar);
+
+		if (SimulcastCVar->GetBool())
+		{
+			// Check that the selected codec supports simulcast
+			FString Codec = CodecCVar->GetString();
+			if (Codec != TEXT("H264") && Codec != TEXT("VP8"))
+			{
+				UE_LOGFMT(LogPixelStreaming2Settings, Warning, "Selected codec doesn't support simulcast! Resetting default codec to {0}", CodecCVar->GetDefaultValue());
+				CodecCVar->Set(*CodecCVar->GetDefaultValue(), ECVF_SetByConsole);
+			}
+		}
+
+		FString Codec = CodecCVar->GetString();
+		FString ScalabilityMode = ScalabilityModeCVar->GetString();
+		if ((Codec == TEXT("H264") || Codec == TEXT("VP8"))
+			&& (ScalabilityMode != TEXT("L1T1") && ScalabilityMode != TEXT("L1T2") && ScalabilityMode != TEXT("L1T3")))
+		{
+			UE_LOGFMT(LogPixelStreaming2Settings, Warning, "Selected codec doesn't support the {0} scalability mode! Resetting scalability mode to {1}", ScalabilityMode, ScalabilityModeCVar->GetDefaultValue());
+			ScalabilityModeCVar->Set(*ScalabilityModeCVar->GetDefaultValue(), ECVF_SetByConsole);
+		}
+	}
+
+	static void VerifyCodecPreferenceSettings(UPixelStreaming2PluginSettings* This, FProperty* Property, IConsoleVariable* CVar, const FString& Value = TEXT(""))
+	{
+		TArray<FString> ValidStringArray;
+
+		if (Value != "")
+		{
+			TArray<FString> StringArray;
+			Value.ParseIntoArray(StringArray, TEXT(","));
+
+			for (const FString& CodecString : StringArray)
+			{
+				if (StaticEnum<EVideoCodec>()->GetIndexByNameString(CodecString) == INDEX_NONE)
+				{
+					UE_LOGFMT(LogPixelStreaming2Settings, Warning, "Invalid value {0} received for enum of type {1}", CodecString, StaticEnum<EVideoCodec>()->GetName());
+					continue;
+				}
+
+				ValidStringArray.Add(CodecString);
+			}
+		}
+		else
+		{
+			FArrayProperty*		ArrayProperty = CastField<FArrayProperty>(Property);
+			TArray<EVideoCodec> CodecArray = *ArrayProperty->ContainerPtrToValuePtr<TArray<EVideoCodec>>(This);
+			for (EVideoCodec Codec : CodecArray)
+			{
+				ValidStringArray.Add(UE::PixelStreaming2::GetCVarStringFromEnum(Codec));
+			}
+		}
+
+		CVar->Set(*FString::Join(ValidStringArray, TEXT(",")), ECVF_SetByCommandline);
+	}
+
+	FString ConsoleVariableToCommandArgValue(const FString InCVarName)
+	{
+		// CVars are . deliminated by section. To get their equivilent commandline arg for parsing
+		// we need to remove the . and add a "="
+		return InCVarName.Replace(TEXT("."), TEXT("")).Replace(TEXT("PixelStreaming2"), TEXT("PixelStreaming")).Append(TEXT("="));
+	}
+
+	FString ConsoleVariableToCommandArgParam(const FString InCVarName)
+	{
+		// CVars are . deliminated by section. To get their equivilent commandline arg parameter, we need to to remove the .
+		return InCVarName.Replace(TEXT("."), TEXT("")).Replace(TEXT("PixelStreaming2"), TEXT("PixelStreaming"));
+	}
+
+	static void ParseLegacyCommandLineValue(const TCHAR* Match, TAutoConsoleVariable<FString>& CVar)
+	{
+		FString Value;
+		if (FParse::Value(FCommandLine::Get(), Match, Value))
+		{
+			CVar->Set(*Value, ECVF_SetByCommandline);
+		}
+	};
+
+	static void ParseLegacyCommandLineOption(const TCHAR* Match, TAutoConsoleVariable<bool>& CVar)
+	{
+		FString ValueMatch(Match);
+		ValueMatch.Append(TEXT("="));
+		FString Value;
+		if (FParse::Value(FCommandLine::Get(), *ValueMatch, Value))
+		{
+			if (Value.Equals(FString(TEXT("true")), ESearchCase::IgnoreCase))
+			{
+				CVar->Set(true, ECVF_SetByCommandline);
+			}
+			else if (Value.Equals(FString(TEXT("false")), ESearchCase::IgnoreCase))
+			{
+				CVar->Set(false, ECVF_SetByCommandline);
+			}
+		}
+		else if (FParse::Param(FCommandLine::Get(), Match))
+		{
+			CVar->Set(true, ECVF_SetByCommandline);
+		}
+	}
+} // namespace
+
+static FName PixelStreaming2ConsoleVariableMetaFName(TEXT("ConsoleVariable"));
+static FName PixelStreaming2MappedConsoleVariableFName(TEXT("MappedConsoleVariable"));
+
+// Begin Pixel Streaming Plugin CVars
+TAutoConsoleVariable<bool> UPixelStreaming2PluginSettings::CVarLogStats(
+	TEXT("PixelStreaming2.LogStats"),
+	false,
+	TEXT("Whether to show PixelStreaming stats in the log (default: false)."),
+	FConsoleVariableDelegate::CreateLambda([](IConsoleVariable* Var) { Delegates()->OnLogStatsChanged.Broadcast(Var); }),
+	ECVF_Default);
+
+TAutoConsoleVariable<bool> UPixelStreaming2PluginSettings::CVarDisableLatencyTester(
+	TEXT("PixelStreaming2.DisableLatencyTester"),
+	false,
+	TEXT("If true disables latency tester being triggerable."),
+	ECVF_Default);
+
+TAutoConsoleVariable<FString> UPixelStreaming2PluginSettings::CVarInputController(
+	TEXT("PixelStreaming2.InputController"),
+	TEXT("Any"),
+	TEXT("Various modes of input control supported by Pixel Streaming, currently: \"Any\"  or \"Host\". Default: Any"),
+	ECVF_Default);
+
+TAutoConsoleVariable<bool> UPixelStreaming2PluginSettings::CVarDecoupleFramerate(
+	TEXT("PixelStreaming2.DecoupleFramerate"),
+	false,
+	TEXT("Whether we should only stream as fast as we render or at some fixed interval. Coupled means only stream what we render."),
+	FConsoleVariableDelegate::CreateLambda([](IConsoleVariable* Var) { Delegates()->OnDecoupleFramerateChanged.Broadcast(Var); }),
+	ECVF_Default);
+
+TAutoConsoleVariable<float> UPixelStreaming2PluginSettings::CVarDecoupleWaitFactor(
+	TEXT("PixelStreaming2.DecoupleWaitFactor"),
+	1.25f,
+	TEXT("Frame rate factor to wait for a captured frame when streaming in decoupled mode. Higher factor waits longer but may also result in higher latency."),
+	ECVF_Default);
+
+TAutoConsoleVariable<float> UPixelStreaming2PluginSettings::CVarSignalingReconnectInterval(
+	TEXT("PixelStreaming2.SignalingReconnectInterval"),
+	2.0f,
+	TEXT("Changes the number of seconds between attempted reconnects to the signaling server. This is useful for reducing the log spam produced from attempted reconnects. A value <= 0 results in no reconnect. Default: 2.0s"),
+	ECVF_Default);
+
+TAutoConsoleVariable<float> UPixelStreaming2PluginSettings::CVarSignalingKeepAliveInterval(
+	TEXT("PixelStreaming2.SignalingKeepAliveInterval"),
+	30.0f,
+	TEXT("Changes the number of seconds between pings to the signaling server. This is useful for keeping the connection active. A value <= 0 results in no pings. Default: 30.0"),
+	ECVF_Default);
+
+TAutoConsoleVariable<bool> UPixelStreaming2PluginSettings::CVarUseMediaCapture(
+	TEXT("PixelStreaming2.UseMediaCapture"),
+	true,
+	TEXT("Use Media Capture from MediaIOFramework to capture frames rather than Pixel Streamings internal backbuffer sources."),
+	ECVF_Default);
+
+TAutoConsoleVariable<FString> UPixelStreaming2PluginSettings::CVarDefaultStreamerID(
+	TEXT("PixelStreaming2.ID"),
+	TEXT("DefaultStreamer"),
+	TEXT("Default Streamer ID to be used when not specified elsewhere."),
+	ECVF_Default);
+
+TAutoConsoleVariable<FString> UPixelStreaming2PluginSettings::CVarConnectionURL(
+	TEXT("PixelStreaming2.ConnectionURL"),
+	TEXT(""),
+	TEXT("Default URL to connect to. This can be a URL to a signalling server or some other endpoint with the format (protocol)://(host):(port)"),
+	ECVF_Default);
+
+FAutoConsoleVariableDeprecated UPixelStreaming2PluginSettings::CVarSignallingURL(TEXT("PixelStreaming2.SignallingURL"), TEXT("PixelStreaming2.ConnectionURL"), TEXT("5.6"));
+
+TAutoConsoleVariable<bool> UPixelStreaming2PluginSettings::CVarCaptureUseFence(
+	TEXT("PixelStreaming2.CaptureUseFence"),
+	true,
+	TEXT("Whether the texture copy we do during image capture should use a fence or not (non-fenced is faster but less safe)."),
+	FConsoleVariableDelegate::CreateLambda([](IConsoleVariable* Var) { Delegates()->OnCaptureUseFenceChanged.Broadcast(Var); }),
+	ECVF_Default);
+
+TAutoConsoleVariable<bool> UPixelStreaming2PluginSettings::CVarDebugDumpAudio(
+	TEXT("PixelStreaming2.DumpDebugAudio"),
+	false,
+	TEXT("Dumps mixed audio from PS2 to a file on disk for debugging purposes."),
+	FConsoleVariableDelegate::CreateLambda([](IConsoleVariable* Var) { Delegates()->OnDebugDumpAudioChanged.Broadcast(Var); }),
+	ECVF_Default);
+
+// Begin Encoder CVars
+
+TAutoConsoleVariable<int32> UPixelStreaming2PluginSettings::CVarEncoderTargetBitrate(
+	TEXT("PixelStreaming2.Encoder.TargetBitrate"),
+	-1,
+	TEXT("Target bitrate (bps). Ignore the bitrate WebRTC wants (not recommended). Set to -1 to disable. Default -1."),
+	ECVF_RenderThreadSafe);
+
+TAutoConsoleVariable<int32> UPixelStreaming2PluginSettings::CVarEncoderMinQuality(
+	TEXT("PixelStreaming2.Encoder.MinQuality"),
+	0,
+	TEXT("0-100, Higher values result in a better minimum quality but higher average bitrates. Default 0 - i.e. no limit on a minimum Quality."),
+	ECVF_Default);
+
+TAutoConsoleVariable<int32> UPixelStreaming2PluginSettings::CVarEncoderMaxQuality(
+	TEXT("PixelStreaming2.Encoder.MaxQuality"),
+	100,
+	TEXT("0-100, Lower values result in lower average bitrates but reduces maximum achievable quality. Default 100 - i.e. no limit on a maximum Quality."),
+	ECVF_Default);
+
+TAutoConsoleVariable<FString> UPixelStreaming2PluginSettings::CVarEncoderQualityPreset(
+	TEXT("PixelStreaming2.Encoder.QualityPreset"),
+	TEXT("Default"),
+	TEXT("PixelStreaming encoder presets that affecting Quality vs Bitrate. Supported modes are: `ULTRA_LOW_QUALITY`, `LOW_QUALITY`, `DEFAULT`, `HIGH_QUALITY` or `LOSSLESS`"),
+	FConsoleVariableDelegate::CreateStatic(&CheckConsoleEnum<EAVPreset>),
+	ECVF_Default);
+
+TAutoConsoleVariable<FString> UPixelStreaming2PluginSettings::CVarEncoderLatencyMode(
+	TEXT("PixelStreaming2.Encoder.LatencyMode"),
+	TEXT("UltraLowLatency"),
+	TEXT("PixelStreaming encoder mode that affecting Quality vs Latency. Supported modes are: `ULTRA_LOW_LATENCY`, `LOW_LATENCY` or `DEFAULT`"),
+	ECVF_Default);
+
+TAutoConsoleVariable<int32> UPixelStreaming2PluginSettings::CVarEncoderKeyframeInterval(
+	TEXT("PixelStreaming2.Encoder.KeyframeInterval"),
+	-1,
+	TEXT("How many frames before a key frame is sent. Default: -1 which disables the sending of periodic key frames. Note: NVENC reqires a reinitialization when this changes."),
+	ECVF_Default);
+
+TAutoConsoleVariable<int32> UPixelStreaming2PluginSettings::CVarEncoderMaxSessions(
+	TEXT("PixelStreaming2.Encoder.MaxSessions"),
+	-1,
+	TEXT("-1 implies no limit. Maximum number of concurrent hardware encoder sessions for Pixel Streaming. Note GeForce gpus only support 8 concurrent sessions and will rollover to software encoding when that number is exceeded."),
+	ECVF_Default);
+
+TAutoConsoleVariable<bool> UPixelStreaming2PluginSettings::CVarEncoderEnableSimulcast(
+	TEXT("PixelStreaming2.Encoder.EnableSimulcast"),
+	false,
+	TEXT("Enables simulcast. When enabled, the encoder will encode at full resolution, 1/2 resolution and 1/4 resolution simultaneously. Note: Simulcast is only supported with `H264` and `VP8` and you must use the SFU from the infrastructure to fully utilise this functionality."),
+	FConsoleVariableDelegate::CreateLambda([](IConsoleVariable* Var) { VerifyCVarVideoSettings(nullptr); Delegates()->OnSimulcastEnabledChanged.Broadcast(Var); }),
+	ECVF_Default);
+
+TAutoConsoleVariable<FString> UPixelStreaming2PluginSettings::CVarEncoderCodec(
+	TEXT("PixelStreaming2.Encoder.Codec"),
+	"H264",
+	TEXT("PixelStreaming default encoder codec. Supported values are: `H264`, `VP8`, `VP9` or `AV1`"),
+	FConsoleVariableDelegate::CreateStatic(&VerifyCVarVideoSettings),
+	ECVF_Default);
+
+TAutoConsoleVariable<FString> UPixelStreaming2PluginSettings::CVarEncoderScalabilityMode(
+	TEXT("PixelStreaming2.Encoder.ScalabilityMode"),
+	TEXT("L1T1"),
+	TEXT("Indicates number of Spatial and temporal layers used, default: L1T1. For a full list of values refer to https://www.w3.org/TR/webrtc-svc/#scalabilitymodes*"),
+	FConsoleVariableDelegate::CreateLambda([](IConsoleVariable* Var) { VerifyCVarVideoSettings(nullptr); Delegates()->OnScalabilityModeChanged.Broadcast(Var); }),
+	ECVF_Default);
+
+TAutoConsoleVariable<FString> UPixelStreaming2PluginSettings::CVarEncoderH264Profile(
+	TEXT("PixelStreaming2.Encoder.H264Profile"),
+	TEXT("Baseline"),
+	TEXT("PixelStreaming encoder profile. Supported modes are: `AUTO`, `BASELINE`, `MAIN`, `HIGH`, `PROGRESSIVE_HIGH`, `CONSTRAINED_HIGH` or `HIGH444`"),
+	ECVF_Default);
+
+TAutoConsoleVariable<bool> UPixelStreaming2PluginSettings::CVarEncoderDebugDumpFrame(
+	TEXT("PixelStreaming2.Encoder.DumpDebugFrames"),
+	false,
+	TEXT("Dumps frames from the encoder to a file on disk for debugging purposes."),
+	FConsoleVariableDelegate::CreateLambda([](IConsoleVariable* Var) { Delegates()->OnEncoderDebugDumpFrameChanged.Broadcast(Var); }),
+	ECVF_Default);
+
+// Begin WebRTC CVars
+
+TAutoConsoleVariable<int32> UPixelStreaming2PluginSettings::CVarWebRTCFps(
+	TEXT("PixelStreaming2.WebRTC.Fps"),
+	60,
+	TEXT("Framerate for WebRTC encoding. Default: 60"),
+	FConsoleVariableDelegate::CreateLambda([](IConsoleVariable* Var) { Delegates()->OnWebRTCFpsChanged.Broadcast(Var); }),
+	ECVF_Default);
+
+// Note: 1 megabit is the maximum allowed in WebRTC for a start bitrate.
+TAutoConsoleVariable<int32> UPixelStreaming2PluginSettings::CVarWebRTCStartBitrate(
+	TEXT("PixelStreaming2.WebRTC.StartBitrate"),
+	1000000,
+	TEXT("Start bitrate (bps) that WebRTC will try begin the stream with. Must be between Min/Max bitrates. Default: 1000000"),
+	ECVF_RenderThreadSafe);
+
+TAutoConsoleVariable<int32> UPixelStreaming2PluginSettings::CVarWebRTCMinBitrate(
+	TEXT("PixelStreaming2.WebRTC.MinBitrate"),
+	100000,
+	TEXT("Min bitrate (bps) that WebRTC will not request below. Careful not to set too high otherwise WebRTC will just drop frames. Default: 100000"),
+	FConsoleVariableDelegate::CreateLambda([](IConsoleVariable* Var) { Delegates()->OnWebRTCBitrateChanged.Broadcast(Var); }),
+	ECVF_RenderThreadSafe);
+
+TAutoConsoleVariable<int32> UPixelStreaming2PluginSettings::CVarWebRTCMaxBitrate(
+	TEXT("PixelStreaming2.WebRTC.MaxBitrate"),
+	40000000,
+	TEXT("Max bitrate (bps) that WebRTC will not request above. Default: 40000000 aka 40 megabits/per second."),
+	FConsoleVariableDelegate::CreateLambda([](IConsoleVariable* Var) { Delegates()->OnWebRTCBitrateChanged.Broadcast(Var); }),
+	ECVF_RenderThreadSafe);
+
+TAutoConsoleVariable<bool> UPixelStreaming2PluginSettings::CVarWebRTCDisableReceiveAudio(
+	TEXT("PixelStreaming2.WebRTC.DisableReceiveAudio"),
+	false,
+	TEXT("Disables receiving audio from the browser into UE."),
+	ECVF_Default);
+
+TAutoConsoleVariable<bool> UPixelStreaming2PluginSettings::CVarWebRTCDisableReceiveVideo(
+	TEXT("PixelStreaming2.WebRTC.DisableReceiveVideo"),
+	true,
+	TEXT("Disables receiving video from the browser into UE."),
+	ECVF_Default);
+
+TAutoConsoleVariable<bool> UPixelStreaming2PluginSettings::CVarWebRTCDisableTransmitAudio(
+	TEXT("PixelStreaming2.WebRTC.DisableTransmitAudio"),
+	false,
+	TEXT("Disables transmission of UE audio to the browser."),
+	ECVF_Default);
+
+TAutoConsoleVariable<bool> UPixelStreaming2PluginSettings::CVarWebRTCDisableTransmitVideo(
+	TEXT("PixelStreaming2.WebRTC.DisableTransmitVideo"),
+	false,
+	TEXT("Disables transmission of UE video to the browser."),
+	ECVF_Default);
+
+TAutoConsoleVariable<bool> UPixelStreaming2PluginSettings::CVarWebRTCDisableAudioSync(
+	TEXT("PixelStreaming2.WebRTC.DisableAudioSync"),
+	true,
+	TEXT("Disables the synchronization of audio and video tracks in WebRTC. This can be useful in low latency usecases where synchronization is not required."),
+	ECVF_Default);
+
+TAutoConsoleVariable<bool> UPixelStreaming2PluginSettings::CVarWebRTCEnableFlexFec(
+	TEXT("PixelStreaming2.WebRTC.EnableFlexFec"),
+	false,
+	TEXT("Signals support for Flexible Forward Error Correction to WebRTC. This can cause a reduction in quality if total bitrate is low."),
+	ECVF_Default);
+
+TAutoConsoleVariable<bool> UPixelStreaming2PluginSettings::CVarWebRTCDisableStats(
+	TEXT("PixelStreaming2.WebRTC.DisableStats"),
+	false,
+	TEXT("Disables the collection of WebRTC stats."),
+	FConsoleVariableDelegate::CreateLambda([](IConsoleVariable* Var) { Delegates()->OnWebRTCDisableStatsChanged.Broadcast(Var); }),
+	ECVF_Default);
+
+TAutoConsoleVariable<float> UPixelStreaming2PluginSettings::CVarWebRTCStatsInterval(
+	TEXT("PixelStreaming2.WebRTC.StatsInterval"),
+	1.f,
+	TEXT("Configures how often WebRTC stats are collected in seconds. Values less than 0.0f disable stats collection. Default: 1.0f"),
+	ECVF_Default);
+
+TAutoConsoleVariable<bool> UPixelStreaming2PluginSettings::CVarWebRTCNegotiateCodecs(
+	TEXT("PixelStreaming2.WebRTC.NegotiateCodecs"),
+	false,
+	TEXT("Whether PS should send all its codecs during sdp handshake so peers can negotiate or just send a single selected codec."),
+	ECVF_Default);
+
+TAutoConsoleVariable<FString> UPixelStreaming2PluginSettings::CVarWebRTCCodecPreferences(
+	TEXT("PixelStreaming2.WebRTC.CodecPreferences"),
+	TEXT("AV1,H264,VP9,VP8"),
+	TEXT("A comma separated list of video codecs specifying the prefered order PS will signal during sdp handshake"),
+	ECVF_Default);
+
+TAutoConsoleVariable<float> UPixelStreaming2PluginSettings::CVarWebRTCAudioGain(
+	TEXT("PixelStreaming2.WebRTC.AudioGain"),
+	1.0f,
+	TEXT("Sets the amount of gain to apply to audio. Default: 1.0"),
+	ECVF_Default);
+
+// End WebRTC CVars
+
+// Begin EditorStreaming CVars
+TAutoConsoleVariable<bool> UPixelStreaming2PluginSettings::CVarEditorStartOnLaunch(
+	TEXT("PixelStreaming2.Editor.StartOnLaunch"),
+	false,
+	TEXT("Start Editor Streaming as soon as the Unreal Editor is launched. Default: false"),
+	ECVF_Default);
+
+TAutoConsoleVariable<bool> UPixelStreaming2PluginSettings::CVarEditorUseRemoteSignallingServer(
+	TEXT("PixelStreaming2.Editor.UseRemoteSignallingServer"),
+	false,
+	TEXT("Enables the use of a remote signalling server. Default: false"),
+	ECVF_Default);
+
+TAutoConsoleVariable<FString> UPixelStreaming2PluginSettings::CVarEditorSource(
+	TEXT("PixelStreaming2.Editor.Source"),
+	TEXT("Editor"),
+	TEXT("Editor PixelStreaming source. Supported values are `Editor`, `LevelEditorViewport`. Default: `Editor`"),
+	FConsoleVariableDelegate::CreateStatic(&CheckConsoleEnum<EPixelStreaming2EditorStreamTypes>),
+	ECVF_Default);
+// End EditorStreaming CVars
+
+// Begin HMD CVars
+TAutoConsoleVariable<bool> UPixelStreaming2PluginSettings::CVarHMDEnable(
+	TEXT("PixelStreaming2.HMD.Enable"),
+	false,
+	TEXT("Enables HMD specific functionality for Pixel Streaming. Namely input handling and stereoscopic rendering. Default: false"),
+	ECVF_Default);
+
+TAutoConsoleVariable<bool> UPixelStreaming2PluginSettings::CVarHMDMatchAspectRatio(
+	TEXT("PixelStreaming2.HMD.MatchAspectRatio"),
+	true,
+	TEXT("If true automatically resize the rendering resolution to match the aspect ratio determined by the HFoV and VFoV. Default: true"),
+	ECVF_Default);
+
+TAutoConsoleVariable<bool> UPixelStreaming2PluginSettings::CVarHMDApplyEyePosition(
+	TEXT("PixelStreaming2.HMD.ApplyEyePosition"),
+	true,
+	TEXT("If true automatically position each eye's rendering by whatever amount WebXR reports for each left-right XRView. If false do no eye positioning. Default: true"),
+	ECVF_Default);
+
+TAutoConsoleVariable<bool> UPixelStreaming2PluginSettings::CVarHMDApplyEyeRotation(
+	TEXT("PixelStreaming2.HMD.ApplyEyeRotation"),
+	true,
+	TEXT("If true automatically rotate each eye's rendering by whatever amount WebXR reports for each left-right XRView. If false do no eye rotation. Default: true"),
+	ECVF_Default);
+
+TAutoConsoleVariable<float> UPixelStreaming2PluginSettings::CVarHMDHFOV(
+	TEXT("PixelStreaming2.HMD.HFOV"),
+	-1.0f,
+	TEXT("Overrides the horizontal field of view for HMD rendering, values are in degrees and values less than 0.0f disable the override. Default: -1.0f"),
+	ECVF_Default);
+
+TAutoConsoleVariable<float> UPixelStreaming2PluginSettings::CVarHMDVFOV(
+	TEXT("PixelStreaming2.HMD.VFOV"),
+	-1.0f,
+	TEXT("Overrides the vertical field of view for HMD rendering, values are in degrees and values less than 0.0f disable the override. Default: -1.0f"),
+	ECVF_Default);
+
+TAutoConsoleVariable<float> UPixelStreaming2PluginSettings::CVarHMDIPD(
+	TEXT("PixelStreaming2.HMD.IPD"),
+	-1.0f,
+	TEXT("Overrides the HMD IPD (interpupillary distance), values are in centimeters and values less than 0.0f disable the override. Default: -1.0f"),
+	ECVF_Default);
+
+TAutoConsoleVariable<float> UPixelStreaming2PluginSettings::CVarHMDProjectionOffsetX(
+	TEXT("PixelStreaming2.HMD.ProjectionOffsetX"),
+	-1.0f,
+	TEXT("Overrides the left/right eye projection matrix x-offset, values are in clip space and values less than 0.0f disable the override. Default: -1.0f"),
+	ECVF_Default);
+
+TAutoConsoleVariable<float> UPixelStreaming2PluginSettings::CVarHMDProjectionOffsetY(
+	TEXT("PixelStreaming2.HMD.ProjectionOffsetY"),
+	-1.0f,
+	TEXT("Overrides the left-right eye projection matrix y-offset, values are in clip space and values less than 0.0f disable the override. Default: -1.0f"),
+	ECVF_Default);
+// End HMD CVars
+
+// Begin Input CVars
+TAutoConsoleVariable<bool> UPixelStreaming2PluginSettings::CVarInputAllowConsoleCommands(
+	TEXT("PixelStreaming2.AllowPixelStreamingCommands"),
+	false,
+	TEXT("If true browser can send consoleCommand payloads that execute in UE's console. Default: false"),
+	ECVF_Default);
+
+TAutoConsoleVariable<FString> UPixelStreaming2PluginSettings::CVarInputKeyFilter(
+	TEXT("PixelStreaming2.KeyFilter"),
+	"",
+	TEXT("Comma separated list of keys to ignore from streaming clients. Default: \"\""),
+	FConsoleVariableDelegate::CreateLambda([](IConsoleVariable* Var) { Delegates()->OnInputKeyFilterChanged.Broadcast(Var); }),
+	ECVF_Default);
+// End Input CVars
+
+TArray<EVideoCodec> UPixelStreaming2PluginSettings::GetCodecPreferences()
+{
+	TArray<EVideoCodec> OutCodecPreferences;
+	FString				StringOptions = UPixelStreaming2PluginSettings::CVarWebRTCCodecPreferences.GetValueOnAnyThread();
+	if (StringOptions.IsEmpty())
+	{
+		return OutCodecPreferences;
+	}
+
+	TArray<FString> CodecArray;
+	StringOptions.ParseIntoArray(CodecArray, TEXT(","), true);
+	for (const FString& CodecString : CodecArray)
+	{
+		uint64 EnumIndex = StaticEnum<EVideoCodec>()->GetIndexByNameString(CodecString);
+		checkf(EnumIndex != INDEX_NONE, TEXT("CVar was not containing valid enum string"));
+		OutCodecPreferences.Add(static_cast<EVideoCodec>(StaticEnum<EVideoCodec>()->GetValueByIndex(EnumIndex)));
+	}
+
+	return OutCodecPreferences;
+}
+
+EPortAllocatorFlags UPixelStreaming2PluginSettings::GetPortAllocationFlags()
+{
+	EPortAllocatorFlags OutPortAllocatorFlags = EPortAllocatorFlags::None;
+	FString				StringOptions = UPixelStreaming2PluginSettings::CVarWebRTCPortAllocatorFlags.GetValueOnAnyThread();
+	if (StringOptions.IsEmpty())
+	{
+		return OutPortAllocatorFlags;
+	}
+
+	TArray<FString> FlagArray;
+	StringOptions.ParseIntoArray(FlagArray, TEXT(","), true);
+	int OptionCount = FlagArray.Num();
+	while (OptionCount > 0)
+	{
+		FString Flag = FlagArray[OptionCount - 1];
+
+		// Flags must match EpicRtc\Include\epic_rtc\core\connection_config.h
+		if (Flag == "DISABLE_UDP")
+		{
+			OutPortAllocatorFlags |= EPortAllocatorFlags::DisableUdp;
+		}
+		else if (Flag == "DISABLE_STUN")
+		{
+			OutPortAllocatorFlags |= EPortAllocatorFlags::DisableStun;
+		}
+		else if (Flag == "DISABLE_RELAY")
+		{
+			OutPortAllocatorFlags |= EPortAllocatorFlags::DisableRelay;
+		}
+		else if (Flag == "DISABLE_TCP")
+		{
+			OutPortAllocatorFlags |= EPortAllocatorFlags::DisableTcp;
+		}
+		else if (Flag == "ENABLE_IPV6")
+		{
+			OutPortAllocatorFlags |= EPortAllocatorFlags::EnableIPV6;
+		}
+		else if (Flag == "ENABLE_SHARED_SOCKET")
+		{
+			OutPortAllocatorFlags |= EPortAllocatorFlags::EnableSharedSocket;
+		}
+		else if (Flag == "ENABLE_STUN_RETRANSMIT_ATTRIBUTE")
+		{
+			OutPortAllocatorFlags |= EPortAllocatorFlags::EnableStunRetransmitAttribute;
+		}
+		else if (Flag == "DISABLE_ADAPTER_ENUMERATION")
+		{
+			OutPortAllocatorFlags |= EPortAllocatorFlags::DisableAdapterEnumeration;
+		}
+		else if (Flag == "DISABLE_DEFAULT_LOCAL_CANDIDATE")
+		{
+			OutPortAllocatorFlags |= EPortAllocatorFlags::DisableDefaultLocalCandidate;
+		}
+		else if (Flag == "DISABLE_UDP_RELAY")
+		{
+			OutPortAllocatorFlags |= EPortAllocatorFlags::DisableUdpRelay;
+		}
+		else if (Flag == "DISABLE_COSTLY_NETWORKS")
+		{
+			OutPortAllocatorFlags |= EPortAllocatorFlags::DisableCostlyNetworks;
+		}
+		else if (Flag == "ENABLE_IPV6_ON_WIFI")
+		{
+			OutPortAllocatorFlags |= EPortAllocatorFlags::EnableIPV6OnWifi;
+		}
+		else if (Flag == "ENABLE_ANY_ADDRESS_PORTS")
+		{
+			OutPortAllocatorFlags |= EPortAllocatorFlags::EnableAnyAddressPort;
+		}
+		else if (Flag == "DISABLE_LINK_LOCAL_NETWORKS")
+		{
+			OutPortAllocatorFlags |= EPortAllocatorFlags::DisableLinkLocalNetworks;
+		}
+		else
+		{
+			UE_LOGFMT(LogPixelStreaming2Settings, Warning, "Unknown port allocator flag: {0}", Flag);
+		}
+		OptionCount--;
+	}
+
+	return OutPortAllocatorFlags;
+}
+
+void SetPortAllocationCVarFromProperty(UObject* This, FProperty* Property)
+{
+	const FNumericProperty* EnumProperty = CastField<const FNumericProperty>(Property);
+	void*					PropertyAddress = EnumProperty->ContainerPtrToValuePtr<void>(This);
+	EPortAllocatorFlags		CurrentValue = static_cast<EPortAllocatorFlags>(EnumProperty->GetSignedIntPropertyValue(PropertyAddress));
+
+	FString CVarString = TEXT("");
+
+	if (static_cast<bool>(CurrentValue & EPortAllocatorFlags::DisableUdp))
+	{
+		CVarString += "DISABLE_UDP,";
+	}
+
+	if (static_cast<bool>(CurrentValue & EPortAllocatorFlags::DisableStun))
+	{
+		CVarString += "DISABLE_STUN,";
+	}
+
+	if (static_cast<bool>(CurrentValue & EPortAllocatorFlags::DisableRelay))
+	{
+		CVarString += "DISABLE_RELAY,";
+	}
+
+	if (static_cast<bool>(CurrentValue & EPortAllocatorFlags::DisableTcp))
+	{
+		CVarString += "DISABLE_TCP,";
+	}
+
+	if (static_cast<bool>(CurrentValue & EPortAllocatorFlags::EnableIPV6))
+	{
+		CVarString += "ENABLE_IPV6,";
+	}
+
+	if (static_cast<bool>(CurrentValue & EPortAllocatorFlags::EnableSharedSocket))
+	{
+		CVarString += "ENABLE_SHARED_SOCKET,";
+	}
+
+	if (static_cast<bool>(CurrentValue & EPortAllocatorFlags::EnableStunRetransmitAttribute))
+	{
+		CVarString += "ENABLE_STUN_RETRANSMIT_ATTRIBUTE,";
+	}
+
+	if (static_cast<bool>(CurrentValue & EPortAllocatorFlags::DisableAdapterEnumeration))
+	{
+		CVarString += "DISABLE_ADAPTER_ENUMERATION,";
+	}
+
+	if (static_cast<bool>(CurrentValue & EPortAllocatorFlags::DisableDefaultLocalCandidate))
+	{
+		CVarString += "DISABLE_DEFAULT_LOCAL_CANDIDATE,";
+	}
+
+	if (static_cast<bool>(CurrentValue & EPortAllocatorFlags::DisableUdpRelay))
+	{
+		CVarString += "DISABLE_UDP_RELAY,";
+	}
+
+	if (static_cast<bool>(CurrentValue & EPortAllocatorFlags::DisableCostlyNetworks))
+	{
+		CVarString += "DISABLE_COSTLY_NETWORKS,";
+	}
+
+	if (static_cast<bool>(CurrentValue & EPortAllocatorFlags::EnableIPV6OnWifi))
+	{
+		CVarString += "ENABLE_IPV6_ON_WIFI,";
+	}
+
+	if (static_cast<bool>(CurrentValue & EPortAllocatorFlags::EnableAnyAddressPort))
+	{
+		CVarString += "ENABLE_ANY_ADDRESS_PORTS,";
+	}
+
+	if (static_cast<bool>(CurrentValue & EPortAllocatorFlags::DisableLinkLocalNetworks))
+	{
+		CVarString += "DISABLE_LINK_LOCAL_NETWORKS,";
+	}
+
+	UPixelStreaming2PluginSettings::CVarWebRTCPortAllocatorFlags.AsVariable()->Set(*CVarString, ECVF_SetByProjectSetting);
+}
+
+TAutoConsoleVariable<FString> UPixelStreaming2PluginSettings::CVarWebRTCPortAllocatorFlags(
+	TEXT("PixelStreaming2.WebRTC.PortAllocatorFlags"),
+	TEXT(""),
+	TEXT("Sets the WebRTC port allocator flags. Format:\"DISABLE_UDP,DISABLE_STUN,...\""),
+	ECVF_Default);
+
+TAutoConsoleVariable<int> UPixelStreaming2PluginSettings::CVarWebRTCMinPort(
+	TEXT("PixelStreaming2.WebRTC.MinPort"),
+	49152, // Default according to RFC5766
+	TEXT("Sets the minimum usable port for the WebRTC port allocator. Default: 49152"),
+	ECVF_Default);
+
+TAutoConsoleVariable<int> UPixelStreaming2PluginSettings::CVarWebRTCMaxPort(
+	TEXT("PixelStreaming2.WebRTC.MaxPort"),
+	65535, // Default according to RFC5766
+	TEXT("Sets the maximum usable port for the WebRTC port allocator. Default: 65535"),
+	ECVF_Default);
+
+TAutoConsoleVariable<FString> UPixelStreaming2PluginSettings::CVarWebRTCFieldTrials(
+	TEXT("PixelStreaming2.WebRTC.FieldTrials"),
+	TEXT(""),
+	TEXT("Sets the WebRTC field trials string. Format:\"TRIAL1/VALUE1/TRIAL2/VALUE2/\""),
+	ECVF_Default);
+
+TAutoConsoleVariable<bool> UPixelStreaming2PluginSettings::CVarWebRTCDisableFrameDropper(
+	TEXT("PixelStreaming2.WebRTC.DisableFrameDropper"),
+	false,
+	TEXT("Disables the WebRTC internal frame dropper using the field trial WebRTC-FrameDropper/Disabled/"),
+	ECVF_Default);
+
+TAutoConsoleVariable<float> UPixelStreaming2PluginSettings::CVarWebRTCVideoPacingMaxDelay(
+	TEXT("PixelStreaming2.WebRTC.VideoPacing.MaxDelay"),
+	-1.0f,
+	TEXT("Enables the WebRTC-Video-Pacing field trial and sets the max delay (ms) parameter. Default: -1.0f (values below zero are discarded.)"),
+	ECVF_Default);
+
+TAutoConsoleVariable<float> UPixelStreaming2PluginSettings::CVarWebRTCVideoPacingFactor(
+	TEXT("PixelStreaming2.WebRTC.VideoPacing.Factor"),
+	-1.0f,
+	TEXT("Enables the WebRTC-Video-Pacing field trial and sets the video pacing factor parameter. Larger values are more lenient on larger bitrates. Default: -1.0f (values below zero are discarded.)"),
+	ECVF_Default);
+
+UPixelStreaming2PluginSettings::FDelegates* UPixelStreaming2PluginSettings::DelegateSingleton = nullptr;
+
+UPixelStreaming2PluginSettings::~UPixelStreaming2PluginSettings()
+{
+	DelegateSingleton = nullptr;
+}
+
+FName UPixelStreaming2PluginSettings::GetCategoryName() const
+{
+	return TEXT("Plugins");
+}
+
+#if WITH_EDITOR
+FText UPixelStreaming2PluginSettings::GetSectionText() const
+{
+	return NSLOCTEXT("PixelStreaming2Plugin", "PixelStreaming2SettingsSection", "PixelStreaming2");
+}
+
+void UPixelStreaming2PluginSettings::PostEditChangeProperty(FPropertyChangedEvent& PropertyChangedEvent)
+{
+	Super::PostEditChangeProperty(PropertyChangedEvent);
+
+	// NOTE due to legacy variables from the commandline the CVars for settings enums store their string values
+	if (const FEnumProperty* EnumProperty = CastField<const FEnumProperty>(PropertyChangedEvent.Property))
+	{
+		// Enums are not directly alligned in the property map so we get the address
+		FNumericProperty* UnderlyingProp = EnumProperty->GetUnderlyingProperty();
+		void*			  PropertyAddress = EnumProperty->ContainerPtrToValuePtr<void>(this);
+
+		// Get the string value of the change property
+		FString ConsoleString = EnumProperty->GetEnum()->GetNameStringByValue(UnderlyingProp->GetSignedIntPropertyValue(PropertyAddress));
+
+		// If the property has MappedConsoleVariable metadata fetch it and set it to the enums string value
+		if (IConsoleVariable* ConsoleVariable = IConsoleManager::Get().FindConsoleVariable(*PropertyChangedEvent.Property->GetMetaData(PixelStreaming2MappedConsoleVariableFName)))
+		{
+			ConsoleVariable->Set(*ConsoleString, ECVF_SetByConsole);
+		}
+	}
+	else
+	{
+		if (PropertyChangedEvent.Property->HasMetaData("Bitmask"))
+		{
+			if (PropertyChangedEvent.Property->GetNameCPP() == "WebRTCPortAllocatorFlags")
+			{
+				SetPortAllocationCVarFromProperty(this, PropertyChangedEvent.Property);
+			}
+		}
+		else if (PropertyChangedEvent.Property->GetNameCPP() == "WebRTCCodecPreferences")
+		{
+			IConsoleVariable* ConsoleVariable = IConsoleManager::Get().FindConsoleVariable(*PropertyChangedEvent.Property->GetMetaData(PixelStreaming2ConsoleVariableMetaFName));
+			VerifyCodecPreferenceSettings(this, PropertyChangedEvent.Property, ConsoleVariable);
+		}
+		// Codec and ScalabilityMode properties are updated in VerifyVideoSettings once we know all the settings are compatible
+		else if (PropertyChangedEvent.Property->GetNameCPP() != "Codec" && PropertyChangedEvent.Property->GetNameCPP() != "ScalabilityMode")
+		{
+			ExportValuesToConsoleVariables(PropertyChangedEvent.Property);
+		}
+	}
+
+	VerifyVideoSettings();
+}
+
+void UPixelStreaming2PluginSettings::VerifyVideoSettings()
+{
+	FProperty*	   SimulcastProperty = GetClass()->FindPropertyByName(TEXT("EnableSimulcast"));
+	FBoolProperty* SimulcastBoolProperty = CastField<FBoolProperty>(SimulcastProperty);
+	bool		   bSimulcastEnabled = SimulcastBoolProperty->GetPropertyValue_InContainer(this);
+
+	FProperty*	  CodecProperty = GetClass()->FindPropertyByName(TEXT("Codec"));
+	FStrProperty* CodecStrProperty = CastField<FStrProperty>(CodecProperty);
+	FString		  CodecString = CodecStrProperty->GetPropertyValue_InContainer(this);
+
+	FProperty*	  ScalabilityModeProperty = GetClass()->FindPropertyByName(TEXT("ScalabilityMode"));
+	FStrProperty* ScalabilityModeStrProperty = CastField<FStrProperty>(ScalabilityModeProperty);
+	FString		  ScalabilityModeString = ScalabilityModeStrProperty->GetPropertyValue_InContainer(this);
+
+	if (bSimulcastEnabled)
+	{
+		if (CodecString != TEXT("H264") && CodecString != TEXT("VP8"))
+		{
+			UE_LOGFMT(LogPixelStreaming2Settings, Warning, "Default codec ({0}) doesn't support simulcast! Resetting default codec to H.264", CodecString);
+			CodecStrProperty->SetPropertyValue_InContainer(this, TEXT("H264"));
+		}
+	}
+
+	CodecString = CodecStrProperty->GetPropertyValue_InContainer(this);
+	if ((CodecString == TEXT("H264") || CodecString == TEXT("VP8"))
+		&& (ScalabilityModeString != TEXT("L1T1") && ScalabilityModeString != TEXT("L1T2") && ScalabilityModeString != TEXT("L1T3")))
+	{
+		UE_LOGFMT(LogPixelStreaming2Settings, Warning, "Default codec ({0}) doesn't support the {1} scalability mode! Resetting scalability mode to L1T1", CodecString, ScalabilityModeString);
+		ScalabilityModeStrProperty->SetPropertyValue_InContainer(this, TEXT("L1T1"));
+	}
+
+	ExportValuesToConsoleVariables(CodecProperty);
+	ExportValuesToConsoleVariables(ScalabilityModeProperty);
+}
+#endif
+
+void UPixelStreaming2PluginSettings::SetCVarFromPropertyAndValue(IConsoleVariable* CVar, FProperty* Property, const FString& CVarString, const FString& Value)
+{
+	checkf(CVar, TEXT("CVar is nullptr"));
+
+	if (FByteProperty* ByteProperty = CastField<FByteProperty>(Property); ByteProperty != NULL && ByteProperty->Enum != NULL)
+	{
+		int32 CastValue;
+		FParse::Value(FCommandLine::Get(), *ConsoleVariableToCommandArgValue(CVarString), CastValue);
+		CVar->Set(CastValue, ECVF_SetByCommandline);
+	}
+	else if (FEnumProperty* EnumProperty = CastField<FEnumProperty>(Property))
+	{
+		int64 EnumIndex = EnumProperty->GetEnum()->GetIndexByNameString(Value.Replace(TEXT("_"), TEXT("")));
+		if (EnumIndex != INDEX_NONE)
+		{
+			CVar->Set(*EnumProperty->GetEnum()->GetNameStringByIndex(EnumIndex), ECVF_SetByCommandline);
+		}
+		else
+		{
+			UE_LOGFMT(LogPixelStreaming2Settings, Warning, "{0} is not a valid enum value for {1}", Value, CVarString);
+		}
+	}
+	else if (FBoolProperty* BoolProperty = CastField<FBoolProperty>(Property))
+	{
+		if (Value.Equals(FString(TEXT("true")), ESearchCase::IgnoreCase))
+		{
+			CVar->Set(true, ECVF_SetByCommandline);
+		}
+		else if (Value.Equals(FString(TEXT("false")), ESearchCase::IgnoreCase))
+		{
+			CVar->Set(false, ECVF_SetByCommandline);
+		}
+		else if (FParse::Param(FCommandLine::Get(), *ConsoleVariableToCommandArgParam(CVarString)))
+		{
+			CVar->Set(true, ECVF_SetByCommandline);
+		}
+	}
+	else if (FIntProperty* IntProperty = CastField<FIntProperty>(Property))
+	{
+		int32 CastValue;
+		FParse::Value(FCommandLine::Get(), *ConsoleVariableToCommandArgValue(CVarString), CastValue);
+		CVar->Set(CastValue, ECVF_SetByCommandline);
+	}
+	else if (FFloatProperty* FloatProperty = CastField<FFloatProperty>(Property))
+	{
+		float CastValue;
+		FParse::Value(FCommandLine::Get(), *ConsoleVariableToCommandArgValue(CVarString), CastValue);
+		CVar->Set(CastValue, ECVF_SetByCommandline);
+	}
+	else if (FStrProperty* StringProperty = CastField<FStrProperty>(Property))
+	{
+		CVar->Set(*Value, ECVF_SetByCommandline);
+	}
+	else if (FNameProperty* NameProperty = CastField<FNameProperty>(Property))
+	{
+		CVar->Set(*Value, ECVF_SetByCommandline);
+	}
+}
+
+void UPixelStreaming2PluginSettings::SetCVarFromProperty(IConsoleVariable* CVar, FProperty* Property, const FString& CVarString)
+{
+	checkf(CVar, TEXT("CVar is nullptr"));
+
+	if (FByteProperty* ByteProperty = CastField<FByteProperty>(Property); ByteProperty != NULL && ByteProperty->Enum != NULL)
+	{
+		CVar->Set(ByteProperty->GetPropertyValue_InContainer(this), ECVF_SetByCommandline);
+	}
+	else if (FEnumProperty* EnumProperty = CastField<FEnumProperty>(Property))
+	{
+		void* PropertyAddress = EnumProperty->ContainerPtrToValuePtr<void>(this);
+		int64 CurrentValue = EnumProperty->GetUnderlyingProperty()->GetSignedIntPropertyValue(PropertyAddress);
+		CVar->Set(*EnumProperty->GetEnum()->GetNameStringByValue(CurrentValue), ECVF_SetByCommandline);
+	}
+	else if (FBoolProperty* BoolProperty = CastField<FBoolProperty>(Property))
+	{
+		CVar->Set(BoolProperty->GetPropertyValue_InContainer(this), ECVF_SetByCommandline);
+	}
+	else if (FIntProperty* IntProperty = CastField<FIntProperty>(Property))
+	{
+		CVar->Set(IntProperty->GetPropertyValue_InContainer(this), ECVF_SetByCommandline);
+	}
+	else if (FFloatProperty* FloatProperty = CastField<FFloatProperty>(Property))
+	{
+		CVar->Set(FloatProperty->GetPropertyValue_InContainer(this), ECVF_SetByCommandline);
+	}
+	else if (FStrProperty* StringProperty = CastField<FStrProperty>(Property))
+	{
+		CVar->Set(*StringProperty->GetPropertyValue_InContainer(this), ECVF_SetByCommandline);
+	}
+	else if (FNameProperty* NameProperty = CastField<FNameProperty>(Property))
+	{
+		CVar->Set(*NameProperty->GetPropertyValue_InContainer(this).ToString(), ECVF_SetByCommandline);
+	}
+}
+
+// Map of Property Names to their commandline args as GetMetaData() is not avaliable in packaged projects
+static const TMap<FString, FString> GetCmdArg = {
+	{ "LogStats", "PixelStreaming2.LogStats" },
+	{ "SendPlayerIdAsInteger", "PixelStreaming2.SendPlayerIdAsInteger" },
+	{ "DisableLatencyTester", "PixelStreaming2.DisableLatencyTester" },
+	{ "DecoupleFramerate", "PixelStreaming2.DecoupleFrameRate" },
+	{ "DecoupleWaitFactor", "PixelStreaming2.DecoupleWaitFactor" },
+	{ "SignalingReconnectInterval", "PixelStreaming2.SignalingReconnectInterval" },
+	{ "SignalingKeepAliveInterval", "PixelStreaming2.SignalingKeepAliveInterval" },
+	{ "UseMediaCapture", "PixelStreaming2.UseMediaCapture" },
+	{ "DefaultStreamerID", "PixelStreaming2.ID" },
+	{ "ConnectionURL", "PixelStreaming2.ConnectionURL" },
+	{ "CaptureUseFence", "PixelStreaming2.CaptureUseFence" },
+	{ "Codec", "PixelStreaming2.Encoder.Codec" },
+	{ "EncoderTargetBitrate", "PixelStreaming2.Encoder.TargetBitrate" },
+	{ "EncoderMinQuality", "PixelStreaming2.Encoder.MinQuality" },
+	{ "EncoderMaxQuality", "PixelStreaming2.Encoder.MaxQuality" },
+	{ "ScalabilityMode", "PixelStreaming2.Encoder.ScalabilityMode" },
+	{ "KeyframeInterval", "PixelStreaming2.Encoder.KeyframeInterval" },
+	{ "MaxSessions", "PixelStreaming2.Encoder.MaxSessions" },
+	{ "EnableSimulcast", "PixelStreaming2.Encoder.EnableSimulcast" },
+	{ "WebRTCFps", "PixelStreaming2.WebRTC.Fps" },
+	{ "WebRTCStartBitrate", "PixelStreaming2.WebRTC.StartBitrate" },
+	{ "WebRTCMinBitrate", "PixelStreaming2.WebRTC.MinBitrate" },
+	{ "WebRTCMaxBitrate", "PixelStreaming2.WebRTC.MaxBitrate" },
+	{ "WebRTCDisableReceiveAudio", "PixelStreaming2.WebRTC.DisableReceiveAudio" },
+	{ "WebRTCDisableReceiveVideo", "PixelStreaming2.WebRTC.DisableReceiveVideo" },
+	{ "WebRTCDisableTransmitAudio", "PixelStreaming2.WebRTC.DisableTransmitAudio" },
+	{ "WebRTCDisableTransmitVideo", "PixelStreaming2.WebRTC.DisableTransmitVideo" },
+	{ "WebRTCDisableAudioSync", "PixelStreaming2.WebRTC.DisableAudioSync" },
+	{ "WebRTCEnableFlexFec", "PixelStreaming2.WebRTC.EnableFlexFec" },
+	{ "WebRTCDisableStats", "PixelStreaming2.WebRTC.DisableStats" },
+	{ "WebRTCStatsInterval", "PixelStreaming2.WebRTC.StatsInterval" },
+	{ "WebRTCNegotiateCodecs", "PixelStreaming2.WebRTC.NegotiateCodecs" },
+	{ "WebRTCAudioGain", "PixelStreaming2.WebRTC.AudioGain" },
+	{ "WebRTCPortAllocatorFlags", "PixelStreaming2.WebRTC.PortAllocatorFlags" },
+	{ "WebRTCMinPort", "PixelStreaming2.WebRTC.MinPort" },
+	{ "WebRTCMaxPort", "PixelStreaming2.WebRTC.MaxPort" },
+	{ "WebRTCFieldTrials", "PixelStreaming2.WebRTC.FieldTrials" },
+	{ "WebRTCDisableFrameDropper", "PixelStreaming2.WebRTC.DisableFrameDropper" },
+	{ "WebRTCVideoPacingMaxDelay", "PixelStreaming2.WebRTC.VideoPacing.MaxDelay" },
+	{ "WebRTCVideoPacingFactor", "PixelStreaming2.WebRTC.VideoPacing.Factor" },
+	{ "EditorStartOnLaunch", "PixelStreaming2.Editor.StartOnLaunch" },
+	{ "EditorUseRemoteSignallingServer", "PixelStreaming2.Editor.UseRemoteSignallingServer" },
+	{ "HMDEnable", "PixelStreaming2.HMD.Enable" },
+	{ "HMDMatchAspectRatio", "PixelStreaming2.HMD.MatchAspectRatio" },
+	{ "HMDAppleEyePosition", "PixelStreaming2.HMD.ApplyEyePosition" },
+	{ "HMDApplyEyeRotation", "PixelStreaming2.HMD.ApplyEyeRotation" },
+	{ "HMDHFOV", "PixelStreaming2.HMD.HFOV" },
+	{ "HMDVFOV", "PixelStreaming2.HMD.VFOV" },
+	{ "HMDIPD", "PixelStreaming2.HMD.IPD" },
+	{ "HMDProjectionOffsetX", "PixelStreaming2.HMD.ProjectionOffsetX" },
+	{ "HMDProjectionOffsetY", "PixelStreaming2.HMD.ProjectionOffsetY" },
+	{ "InputAllowConsoleCommands", "PixelStreaming2.AllowPixelStreamingCommands" },
+	{ "InputKeyFilter", "PixelStreaming2.KeyFilter" }
+};
+
+static const TMap<FString, FString> GetMappedCmdArg = {
+	{ "InputController", "PixelStreaming2.InputController" },
+	{ "QualityPreset", "PixelStreaming2.Encoder.QualityPreset" },
+	{ "LatencyMode", "PixelStreaming2.Encoder.LatencyMode" },
+	{ "H264Profile", "PixelStreaming2.Encoder.H264Profile" },
+	{ "EditorSource", "PixelStreaming2.Editor.Source" }
+};
+
+//                                                                   Property,   CVar,              CommandLine value
+using FMappingFunc = TFunction<void(UPixelStreaming2PluginSettings*, FProperty*, IConsoleVariable*, const FString&)>;
+
+static const TMap<FString, TPair<FString, FMappingFunc>> GetCustomMappedCmdArg = {
+	{ "WebRTCCodecPreferences", { "PixelStreaming2.WebRTC.CodecPreferences", VerifyCodecPreferenceSettings } },
+};
+
+static const TArray<FString> GetLegacyCmdArg = {
+	"PixelStreaming2.Encoder.MinQp", // Renamed to MaxQuality
+	"PixelStreaming2.Encoder.MaxQp", // Renamed to MinQuality
+	"PixelStreaming2.IP",			 // Moved to URL
+	"PixelStreaming2.Port",			 // Moved to URL
+	"PixelStreaming2.URL",			 // Renamed to ConnectionURL
+	"PixelStreaming2.SignallingURL", // Renamed to ConnectionURL
+	"AllowPixelStreamingCommands",
+	"PixelStreaming2.NegotiateCodecs", // Renamed to WebRTC.NegotiateCodecs
+	"PixelStreaming2.OnScreenStats",   // CVar is removed but launch arg is used in stats.cpp
+	"PixelStreaming2.HudStats",		   // CVar is removed but launch arg is used in stats.cpp
+	"PixelStreaming2.EnableHMD"		   // Renamed to HMDEnable
+};
+
+void UPixelStreaming2PluginSettings::ValidateCommandLineArgs()
+{
+	FString CommandLine = FCommandLine::Get();
+
+	TArray<FString> CommandArray;
+	CommandLine.ParseIntoArray(CommandArray, TEXT(" "), true);
+
+	for (FString Command : CommandArray)
+	{
+		Command.RemoveFromStart(TEXT("-"));
+		if (!Command.StartsWith("PixelStreaming"))
+		{
+			continue;
+		}
+
+		// Get the pure command line arg from an arg that contains an '=', eg PixelStreamingURL=
+		FString CurrentCommandLineArg = Command;
+		if (Command.Contains("="))
+		{
+			Command.Split(TEXT("="), &CurrentCommandLineArg, nullptr);
+		}
+
+		bool bValidArg = false;
+		for (const TPair<FString, FString>& Pair : GetCmdArg)
+		{
+			FString ValidCommandLineArg = ConsoleVariableToCommandArgParam(Pair.Value);
+			if (CurrentCommandLineArg == ValidCommandLineArg)
+			{
+				bValidArg = true;
+			}
+		}
+
+		if (!bValidArg)
+		{
+			for (const TPair<FString, FString>& Pair : GetMappedCmdArg)
+			{
+				FString ValidCommandLineArg = ConsoleVariableToCommandArgParam(Pair.Value);
+				if (CurrentCommandLineArg == ValidCommandLineArg)
+				{
+					bValidArg = true;
+				}
+			}
+		}
+
+		if (!bValidArg)
+		{
+			for (const TPair<FString, TPair<FString, FMappingFunc>>& Pair : GetCustomMappedCmdArg)
+			{
+				FString ValidCommandLineArg = ConsoleVariableToCommandArgParam(Pair.Value.Key);
+				if (CurrentCommandLineArg == ValidCommandLineArg)
+				{
+					bValidArg = true;
+				}
+			}
+		}
+
+		if (!bValidArg)
+		{
+			for (const FString& LegacyArg : GetLegacyCmdArg)
+			{
+				FString ValidCommandLineArg = ConsoleVariableToCommandArgParam(LegacyArg);
+				if (CurrentCommandLineArg == ValidCommandLineArg)
+				{
+					bValidArg = true;
+				}
+			}
+		}
+
+		if (!bValidArg)
+		{
+			UE_LOGFMT(LogPixelStreaming2Settings, Warning, "Unknown PixelStreaming command line arg: {0}", CurrentCommandLineArg);
+		}
+	}
+}
+
+void UPixelStreaming2PluginSettings::ParseLegacyCommandlineArgs()
+{
+	// Begin legacy PixelStreaming command line args
+	int32 MinQP;
+	if (FParse::Value(FCommandLine::Get(), TEXT("PixelStreamingEncoderMinQp="), MinQP))
+	{
+		CVarEncoderMaxQuality.AsVariable()->Set(100.0f * (1.0f - (FMath::Clamp<int32>(MinQP, 0, 51) / 51.0f)), ECVF_SetByCommandline);
+		UE_LOGFMT(LogPixelStreaming2Settings, Log, "PixelStreamingEncoderMinQp is a legacy setting, converted to PixelStreamingEncoderMaxQuality={0}", CVarEncoderMaxQuality.GetValueOnAnyThread());
+	}
+
+	int32 MaxQP;
+	if (FParse::Value(FCommandLine::Get(), TEXT("PixelStreamingEncoderMaxQp="), MaxQP))
+	{
+		CVarEncoderMinQuality.AsVariable()->Set(100.0f * (1.0f - (FMath::Clamp<int32>(MaxQP, 0, 51) / 51.0f)), ECVF_SetByCommandline);
+		UE_LOGFMT(LogPixelStreaming2Settings, Log, "PixelStreamingEncoderMaxQp is a legacy setting, converted to PixelStreamingEncoderMinQuality={0}", CVarEncoderMinQuality.GetValueOnAnyThread());
+	}
+
+	FString LegacyUrl;
+	FString SignallingServerIP;
+	if (FParse::Value(FCommandLine::Get(), TEXT("PixelStreamingIP="), SignallingServerIP))
+	{
+		LegacyUrl += SignallingServerIP;
+	}
+
+	FString SignallingServerPort;
+	if (FParse::Value(FCommandLine::Get(), TEXT("PixelStreamingPort="), SignallingServerPort))
+	{
+		LegacyUrl += TEXT(":") + SignallingServerPort;
+	}
+
+	if (!LegacyUrl.IsEmpty())
+	{
+		CVarConnectionURL.AsVariable()->Set(*(FString("ws://") + LegacyUrl), ECVF_SetByCommandline);
+		UE_LOGFMT(LogPixelStreaming2Settings, Log, "PixelStreamingIP and PixelStreamingPort are legacy settings converted to -PixelStreamingConnectionURL={0}", CVarConnectionURL.GetValueOnAnyThread());
+	}
+
+	// The new URL argument is PixelStreamingConnectionURL= but we want to support the old one too
+	if (FParse::Value(FCommandLine::Get(), TEXT("PixelStreamingURL="), LegacyUrl))
+	{
+		CVarConnectionURL.AsVariable()->Set(*LegacyUrl, ECVF_SetByCommandline);
+	}
+
+	// The new URL argument is PixelStreamingConnectionURL= but we want to support the old one too
+	if (FParse::Value(FCommandLine::Get(), TEXT("PixelStreamingSignallingURL="), LegacyUrl))
+	{
+		CVarConnectionURL.AsVariable()->Set(*LegacyUrl, ECVF_SetByCommandline);
+	}
+
+	ParseLegacyCommandLineOption(TEXT("PixelStreamingNegotiateCodecs"), CVarWebRTCNegotiateCodecs);
+	ParseLegacyCommandLineOption(TEXT("AllowPixelStreamingCommands"), CVarInputAllowConsoleCommands);
+	ParseLegacyCommandLineOption(TEXT("PixelStreamingDebugDumpFrame"), CVarEncoderDebugDumpFrame);
+	// End legacy PixelStreaming command line args
+
+	// Begin legacy PixelStreamingEditor command line args
+	ParseLegacyCommandLineOption(TEXT("EditorPixelStreamingStartOnLaunch"), CVarEditorStartOnLaunch);
+	ParseLegacyCommandLineOption(TEXT("EditorPixelStreamingUseRemoteSignallingServer"), CVarEditorUseRemoteSignallingServer);
+
+	ParseLegacyCommandLineValue(TEXT("EditorPixelStreamingSource="), CVarEditorSource);
+	IConsoleVariable* EditorSourceCVar = IConsoleManager::Get().FindConsoleVariable(TEXT("PixelStreaming2.Editor.Source"));
+	CheckConsoleEnum<EPixelStreaming2EditorStreamTypes>(EditorSourceCVar);
+	// End legacy PixelStreamingEditor command line args
+
+	// End legacy PixelStreamingHMD command line args
+	ParseLegacyCommandLineOption(TEXT("PixelStreamingEnableHMD"), CVarHMDEnable);
+	// End legacy PixelStreamingHMD command line args
+}
+
+void UPixelStreaming2PluginSettings::PostInitProperties()
+{
+	Super::PostInitProperties();
+
+	UE_LOGFMT(LogPixelStreaming2Settings, Log, "Initialising Pixel Streaming settings.");
+
+	ValidateCommandLineArgs();
+
+	for (FProperty* Property = GetClass()->PropertyLink; Property; Property = Property->PropertyLinkNext)
+	{
+		if (!Property->HasAnyPropertyFlags(CPF_Config))
+		{
+			continue;
+		}
+
+		// Handle the majority of commandline argument
+		if (GetCmdArg.Contains(Property->GetNameCPP()))
+		{
+			FString CVarString = GetCmdArg[Property->GetNameCPP()];
+			if (Property->GetNameCPP() == "WebRTCPortAllocatorFlags")
+			{
+				FString ConsoleString;
+				if (FParse::Value(FCommandLine::Get(), *ConsoleVariableToCommandArgValue(CVarString), ConsoleString))
+				{
+					IConsoleManager::Get().FindConsoleVariable(*CVarString)->Set(*ConsoleString, ECVF_SetByCommandline);
+				}
+				else
+				{
+					SetPortAllocationCVarFromProperty(this, Property);
+				}
+				continue;
+			}
+
+			if (!CVarString.IsEmpty())
+			{
+				// Handle a directly parsable commandline
+				FString ConsoleString;
+				if (FParse::Value(FCommandLine::Get(), *ConsoleVariableToCommandArgValue(CVarString), ConsoleString))
+				{
+					SetCVarFromPropertyAndValue(IConsoleManager::Get().FindConsoleVariable(*CVarString), Property, CVarString, ConsoleString);
+				}
+				else if (FParse::Param(FCommandLine::Get(), *ConsoleVariableToCommandArgParam(CVarString)))
+				{
+					SetCVarFromPropertyAndValue(IConsoleManager::Get().FindConsoleVariable(*CVarString), Property, CVarString, TEXT("true"));
+				}
+				else
+				{
+					SetCVarFromProperty(IConsoleManager::Get().FindConsoleVariable(*CVarString), Property, CVarString);
+				}
+			}
+		}
+
+		// Handle a commandline argument that needs mapping from string to enum string
+		if (GetMappedCmdArg.Contains(Property->GetNameCPP()))
+		{
+			FString CVarString = GetMappedCmdArg[Property->GetNameCPP()];
+			FString ConsoleString;
+			if (FParse::Value(FCommandLine::Get(), *ConsoleVariableToCommandArgValue(CVarString), ConsoleString))
+			{
+				SetCVarFromPropertyAndValue(IConsoleManager::Get().FindConsoleVariable(*CVarString), Property, CVarString, ConsoleString);
+			}
+			else
+			{
+				// Safety check that it is actually an EnumProperty
+				if (const FEnumProperty* EnumProperty = CastField<const FEnumProperty>(Property))
+				{
+					FNumericProperty* UnderlyingProp = EnumProperty->GetUnderlyingProperty();
+					void*			  PropertyAddress = EnumProperty->ContainerPtrToValuePtr<void>(this);
+					ConsoleString = EnumProperty->GetEnum()->GetNameStringByValue(UnderlyingProp->GetSignedIntPropertyValue(PropertyAddress));
+
+					if (IConsoleVariable* ConsoleVariable = IConsoleManager::Get().FindConsoleVariable(*CVarString))
+					{
+						ConsoleVariable->Set(*ConsoleString, ECVF_SetByProjectSetting);
+					}
+				}
+				else
+				{
+					checkNoEntry();
+				}
+			}
+		}
+
+		// Handle a commandline argument that needs custom mapping from string to some other type
+		// eg TArray to comma separate FString
+		if (GetCustomMappedCmdArg.Contains(Property->GetNameCPP()))
+		{
+			TPair<FString, FMappingFunc> CVarFuncPair = GetCustomMappedCmdArg[Property->GetNameCPP()];
+			FString						 CVarString = CVarFuncPair.Key;
+			FMappingFunc				 MappingFunc = CVarFuncPair.Value;
+
+			FString ConsoleString;
+			if (FParse::Value(FCommandLine::Get(), *ConsoleVariableToCommandArgValue(CVarString), ConsoleString, false))
+			{
+				// Pass in the console string value. This will set the CVar from what was on the command line
+				MappingFunc(this, Property, IConsoleManager::Get().FindConsoleVariable(*CVarString), ConsoleString);
+			}
+			else
+			{
+				// Pass in an empty value. This will set the CVar from the property's value
+				MappingFunc(this, Property, IConsoleManager::Get().FindConsoleVariable(*CVarString), TEXT(""));
+			}
+		}
+	}
+
+	// Handle parsing of legacy command line args (such as -PixelStreamingUrl) after .ini, properties, and new commandline args.
+	ParseLegacyCommandlineArgs();
+}
+
+UPixelStreaming2PluginSettings::FDelegates* UPixelStreaming2PluginSettings::Delegates()
+{
+	if (DelegateSingleton == nullptr && !IsEngineExitRequested())
+	{
+		DelegateSingleton = new UPixelStreaming2PluginSettings::FDelegates();
+		return DelegateSingleton;
+	}
+	return DelegateSingleton;
+}
+
+TArray<FString> UPixelStreaming2PluginSettings::GetVideoCodecOptions() const
+{
+	FProperty*	   Property = GetClass()->FindPropertyByName(TEXT("EnableSimulcast"));
+	FBoolProperty* BoolProperty = CastField<FBoolProperty>(Property);
+	bool		   bSimulcastEnabled = BoolProperty->GetPropertyValue_InContainer(this);
+
+	if (bSimulcastEnabled)
+	{
+		return {
+			UE::PixelStreaming2::GetCVarStringFromEnum(EVideoCodec::H264),
+			UE::PixelStreaming2::GetCVarStringFromEnum(EVideoCodec::VP8)
+		};
+	}
+
+	return {
+		UE::PixelStreaming2::GetCVarStringFromEnum(EVideoCodec::AV1),
+		UE::PixelStreaming2::GetCVarStringFromEnum(EVideoCodec::H264),
+		UE::PixelStreaming2::GetCVarStringFromEnum(EVideoCodec::VP8),
+		UE::PixelStreaming2::GetCVarStringFromEnum(EVideoCodec::VP9)
+	};
+}
+
+TArray<FString> UPixelStreaming2PluginSettings::GetScalabilityModeOptions() const
+{
+	FProperty*	  Property = GetClass()->FindPropertyByName(TEXT("Codec"));
+	FStrProperty* StrProperty = CastField<FStrProperty>(Property);
+	FString		  SelectedCodec = StrProperty->GetPropertyValue_InContainer(this);
+	// H.264 and VP8 only support temporal scalability
+	bool bRestrictModes = SelectedCodec == TEXT("H264") || SelectedCodec == TEXT("VP8");
+
+	if (bRestrictModes)
+	{
+		return {
+			UE::PixelStreaming2::GetCVarStringFromEnum(EScalabilityMode::L1T1),
+			UE::PixelStreaming2::GetCVarStringFromEnum(EScalabilityMode::L1T2),
+			UE::PixelStreaming2::GetCVarStringFromEnum(EScalabilityMode::L1T3),
+		};
+	}
+
+	TArray<FString> ScalabilityModes;
+	for (uint32 i = 0; i <= static_cast<uint32>(EScalabilityMode::None); i++)
+	{
+		ScalabilityModes.Add(UE::PixelStreaming2::GetCVarStringFromEnum(static_cast<EScalabilityMode>(i)));
+	}
+
+	return ScalabilityModes;
+}
\ No newline at end of file
