[PS2] Fixes to Pixel Streaming 2 decoder displaying old frames and leaking vram.

Problem is WebRTC may send decoded frames to PS2 slower than they are decoded. In addition, ResourcePool is not syncing VideoResources correctly.

The fix undoes CL36622920.

diff --git a/Engine/Plugins/Media/PixelStreaming2/Source/PixelStreaming2RTC/Private/EpicRtcVideoDecoder.cpp b/Engine/Plugins/Media/PixelStreaming2/Source/PixelStreaming2RTC/Private/EpicRtcVideoDecoder.cpp
--- a/Engine/Plugins/Media/PixelStreaming2/Source/PixelStreaming2RTC/Private/EpicRtcVideoDecoder.cpp
+++ b/Engine/Plugins/Media/PixelStreaming2/Source/PixelStreaming2RTC/Private/EpicRtcVideoDecoder.cpp
@@ -116,7 +116,12 @@ namespace UE::PixelStreaming2
 			FAVResult DecodeResult;
 			if constexpr (std::is_same_v<TVideoResource, FVideoResourceRHI>)
 			{
-				FResolvableVideoResourceRHI& DecoderResource = VideoResources.GetOrCreate();
+				// TODO (Eden.Harris) RTCP-7927 Use VideoResources and FetchOrCreate rather than flip flopping.
+				// Note: GetOrCreate currently has a sync issue where old frames are displayed.
+				// By flip flopping VideoResourceIndex, old frames may be overwritten with a new frame texture.
+				// This results in low latency but overwritten frames and resolves old frames accidentally being shown out of order.
+				// FResolvableVideoResourceRHI& DecoderResource = VideoResources.GetOrCreate();
+				FResolvableVideoResourceRHI& DecoderResource = VideoResourcesRHI[VideoResourceIndex++ % 2];
 				DecodeResult = Decoder->ReceiveFrame(DecoderResource);
 
 				if (DecodeResult.IsSuccess())
@@ -126,7 +131,12 @@ namespace UE::PixelStreaming2
 			}
 			else if constexpr (std::is_same_v<TVideoResource, FVideoResourceCPU>)
 			{
-				FResolvableVideoResourceCPU& DecoderResource = VideoResources.GetOrCreate();
+				// TODO (Eden.Harris) RTCP-7927 Use VideoResources and FetchOrCreate rather than flip flopping.
+				// Note: GetOrCreate currently has a sync issue where old frames are displayed.
+				// By flip flopping VideoResourceIndex, old frames may be overwritten with a new frame texture.
+				// This results in low latency but overwritten frames and resolves old frames accidentally being shown out of order.
+				// FResolvableVideoResourceCPU& DecoderResource = VideoResources.GetOrCreate();
+				FResolvableVideoResourceCPU& DecoderResource = VideoResourcesCPU[VideoResourceIndex++ % 2];
 				DecodeResult = Decoder->ReceiveFrame(DecoderResource);
 
 				if (DecodeResult.IsSuccess())
diff --git a/Engine/Plugins/Media/PixelStreaming2/Source/PixelStreaming2RTC/Private/EpicRtcVideoDecoder.h b/Engine/Plugins/Media/PixelStreaming2/Source/PixelStreaming2RTC/Private/EpicRtcVideoDecoder.h
--- a/Engine/Plugins/Media/PixelStreaming2/Source/PixelStreaming2RTC/Private/EpicRtcVideoDecoder.h
+++ b/Engine/Plugins/Media/PixelStreaming2/Source/PixelStreaming2RTC/Private/EpicRtcVideoDecoder.h
@@ -90,9 +90,15 @@ namespace UE::PixelStreaming2
 								 // We only support RHI and CPU resources so if we create one that isn't this type it will error (on purpose)
 								 void>::type>::type; 
 
-		ResourcePoolType VideoResources;
+		// TODO (Eden.Harris) RTCP-7927 Use VideoResources and FetchOrCreate rather than flip flopping
+		// ResourcePoolType VideoResources;
 		// clang-format on
 
+		// TODO (Eden.Harris) RTCP-7927 Use VideoResources and FetchOrCreate rather than flip flopping
+		TStaticArray<FResolvableVideoResourceRHI, 2> VideoResourcesRHI;
+		TStaticArray<FResolvableVideoResourceCPU, 2> VideoResourcesCPU;
+		std::atomic<uint32>							 VideoResourceIndex;
+
 	private:
 		bool LateInitDecoder();
 
