[PS2] Fix: Incorrect stat values

diff --git a/Engine/Plugins/Media/PixelStreaming2/Source/PixelStreaming2RTC/Private/EpicRtcStreamer.cpp b/Engine/Plugins/Media/PixelStreaming2/Source/PixelStreaming2RTC/Private/EpicRtcStreamer.cpp
--- a/Engine/Plugins/Media/PixelStreaming2/Source/PixelStreaming2RTC/Private/EpicRtcStreamer.cpp
+++ b/Engine/Plugins/Media/PixelStreaming2/Source/PixelStreaming2RTC/Private/EpicRtcStreamer.cpp
@@ -737,8 +737,8 @@ namespace UE::PixelStreaming2
 				FStats* Stats = FStats::Get();
 				if (Stats)
 				{
-					Stats->QueryPeerStat(PlayerId, RTCStatCategories::LocalVideoTrack, PixelStreaming2StatNames::MeanEncodeTime, EncodeMs);
-					Stats->QueryPeerStat(PlayerId, RTCStatCategories::LocalVideoTrack, PixelStreaming2StatNames::AvgSendDelay, CaptureToSendMs);
+					Stats->QueryPeerStat(PlayerId, FName(*RTCStatCategories::LocalVideoTrack), PixelStreaming2StatNames::MeanEncodeTime, EncodeMs);
+					Stats->QueryPeerStat(PlayerId, FName(*RTCStatCategories::LocalVideoTrack), PixelStreaming2StatNames::MeanSendDelay, CaptureToSendMs);
 				}
 
 				double TransmissionTimeMs = FPlatformTime::ToMilliseconds64(FPlatformTime::Cycles64());
diff --git a/Engine/Plugins/Media/PixelStreaming2/Source/PixelStreaming2RTC/Private/RTCStatsCollector.cpp b/Engine/Plugins/Media/PixelStreaming2/Source/PixelStreaming2RTC/Private/RTCStatsCollector.cpp
--- a/Engine/Plugins/Media/PixelStreaming2/Source/PixelStreaming2RTC/Private/RTCStatsCollector.cpp
+++ b/Engine/Plugins/Media/PixelStreaming2/Source/PixelStreaming2RTC/Private/RTCStatsCollector.cpp
@@ -31,14 +31,7 @@ namespace UE::PixelStreaming2
 		: AssociatedPlayerId(PlayerId)
 		, LastCalculationCycles(FPlatformTime::Cycles64())
 		, bIsEnabled(!UPixelStreaming2PluginSettings::CVarWebRTCDisableStats.GetValueOnAnyThread())
-		, LocalVideoTrackStatsSink(MakeUnique<FRTPLocalVideoTrackStatsSink>())
-		, LocalAudioTrackStatsSink(MakeUnique<FRTPLocalAudioTrackStatsSink>())
-		, VideoSourceStatsSink(MakeUnique<FVideoSourceStatsSink>())
-		, AudioSourceStatsSink(MakeUnique<FAudioSourceStatsSink>())
-		, RemoteVideoTrackStatsSink(MakeUnique<FRTPRemoteTrackStatsSink>(RTCStatCategories::RemoteVideoTrack))
-		, RemoteAudioTrackStatsSink(MakeUnique<FRTPRemoteTrackStatsSink>(RTCStatCategories::RemoteAudioTrack))
-		, DataTrackStatsSink(MakeUnique<FDataTrackStatsSink>())
-		, CandidatePairStatsSink(MakeUnique<FCandidatePairStatsSink>())
+		, CandidatePairStatsSink(MakeUnique<FCandidatePairStatsSink>(FName(*RTCStatCategories::CandidatePair)))
 	{
 	}
 
@@ -58,53 +51,135 @@ namespace UE::PixelStreaming2
 		uint64 CyclesNow = FPlatformTime::Cycles64();
 		double SecondsDelta = FGenericPlatformTime::ToSeconds64(CyclesNow - LastCalculationCycles);
 
-		if (InStats._localVideoTracks._size > 0)
+		// Local video stats
+		for (uint64 i = 0; i < InStats._localVideoTracks._size; i++)
 		{
-			if (InStats._localVideoTracks._size > 1)
+			const EpicRtcLocalVideoTrackStats& LocalVideoTrackStats = InStats._localVideoTracks._ptr[i];
+
+			// Process video source stats
+			if (!VideoSourceSinks.Contains(i))
+			{
+				FName SinkName = FName(*FString::Printf(TEXT("%s [%u]"), *RTCStatCategories::VideoSource, i));
+				VideoSourceSinks.Add(i, MakeUnique<FVideoSourceStatsSink>(SinkName));
+			}
+
+			VideoSourceSinks[i]->Process(LocalVideoTrackStats._source, AssociatedPlayerId, SecondsDelta);
+
+			// Process video track rtp stats
+			if (!LocalVideoTrackSinks.Contains(i))
+			{
+				LocalVideoTrackSinks.Add(i, {});
+			}
+
+			TMap<uint32, TUniquePtr<FRTPLocalVideoTrackStatsSink>>& SsrcSinks = LocalVideoTrackSinks[i];
+			for (int j = 0; j < LocalVideoTrackStats._rtp._size; j++)
 			{
-				UE_LOGFMT(LogPixelStreaming2RTC, Warning, "Multiple local video tracks is unsupported. Expected: 1, Actual: {0}", InStats._localVideoTracks._size);
+				const EpicRtcLocalTrackRtpStats& RtpStats = LocalVideoTrackStats._rtp._ptr[j];
+
+				if (!SsrcSinks.Contains(RtpStats._local._ssrc))
+				{
+					FName SinkName = FName(*FString::Printf(TEXT("%s [%u] (%u)"), *RTCStatCategories::LocalVideoTrack, i, RtpStats._local._ssrc));
+					SsrcSinks.Add(RtpStats._local._ssrc, MakeUnique<FRTPLocalVideoTrackStatsSink>(SinkName));
+				}
+
+				SsrcSinks[RtpStats._local._ssrc]->Process(RtpStats, AssociatedPlayerId, SecondsDelta);
 			}
-			LocalVideoTrackStatsSink->Process(InStats._localVideoTracks._ptr[0], AssociatedPlayerId, SecondsDelta);
-			VideoSourceStatsSink->Process(InStats._localVideoTracks._ptr[0]._source, AssociatedPlayerId, SecondsDelta);
 		}
 
-		if (InStats._localAudioTracks._size > 0)
+		// Local audio stats
+		for (uint64 i = 0; i < InStats._localAudioTracks._size; i++)
 		{
-			if (InStats._localAudioTracks._size > 1)
+			const EpicRtcLocalAudioTrackStats& LocalAudioTrackStats = InStats._localAudioTracks._ptr[i];
+
+			// Process audio source stats
+			if (!AudioSourceSinks.Contains(i))
+			{
+				FName SinkName = FName(*FString::Printf(TEXT("%s [%u]"), *RTCStatCategories::AudioSource, i));
+				AudioSourceSinks.Add(i, MakeUnique<FAudioSourceStatsSink>(SinkName));
+			}
+
+			AudioSourceSinks[i]->Process(LocalAudioTrackStats._source, AssociatedPlayerId, SecondsDelta);
+
+			// Process audio track rtp stats
+			if (!LocalAudioTrackSinks.Contains(i))
 			{
-				UE_LOGFMT(LogPixelStreaming2RTC, Warning, "Multiple local audio tracks is unsupported. Expected: 1, Actual: {0}", InStats._localAudioTracks._size);
+				LocalAudioTrackSinks.Add(i, {});
 			}
-			LocalAudioTrackStatsSink->Process(InStats._localAudioTracks._ptr[0]._rtp, AssociatedPlayerId, SecondsDelta);
-			AudioSourceStatsSink->Process(InStats._localAudioTracks._ptr[0]._source, AssociatedPlayerId, SecondsDelta);
+
+			TMap<uint32, TUniquePtr<FRTPLocalAudioTrackStatsSink>>& SsrcSinks = LocalAudioTrackSinks[i];
+			const EpicRtcLocalTrackRtpStats&						  RtpStats = LocalAudioTrackStats._rtp;
+
+			if (!SsrcSinks.Contains(RtpStats._local._ssrc))
+			{
+				FName SinkName = FName(*FString::Printf(TEXT("%s [%u] (%u)"), *RTCStatCategories::LocalAudioTrack, i, RtpStats._local._ssrc));
+				SsrcSinks.Add(RtpStats._local._ssrc, MakeUnique<FRTPLocalAudioTrackStatsSink>(SinkName));
+			}
+
+			SsrcSinks[RtpStats._local._ssrc]->Process(RtpStats, AssociatedPlayerId, SecondsDelta);
 		}
 
-		if (InStats._remoteVideoTracks._size > 0)
+		// remote video stats
+		for (uint64 i = 0; i < InStats._remoteVideoTracks._size; i++)
 		{
-			if (InStats._remoteVideoTracks._size > 1)
+			const EpicRtcRemoteTrackStats& RemoteVideoTrackStats = InStats._remoteVideoTracks._ptr[i];
+
+			// Process video track rtp stats
+			if (!RemoteVideoTrackSinks.Contains(i))
 			{
-				UE_LOGFMT(LogPixelStreaming2RTC, Warning, "Multiple remote video tracks is unsupported. Expected: 1, Actual: {0}", InStats._remoteVideoTracks._size);
+				RemoteVideoTrackSinks.Add(i, {});
 			}
-			RemoteVideoTrackStatsSink->Process(InStats._remoteVideoTracks._ptr[0]._rtp, AssociatedPlayerId, SecondsDelta);
+
+			TMap<uint32, TUniquePtr<FRTPRemoteTrackStatsSink>>& SsrcSinks = RemoteVideoTrackSinks[i];
+			const EpicRtcRemoteTrackRtpStats&					  RtpStats = RemoteVideoTrackStats._rtp;
+
+			if (!SsrcSinks.Contains(RtpStats._local._ssrc))
+			{
+				FName SinkName = FName(*FString::Printf(TEXT("%s [%u] (%u)"), *RTCStatCategories::RemoteVideoTrack, i, RtpStats._local._ssrc));
+				SsrcSinks.Add(RtpStats._local._ssrc, MakeUnique<FRTPRemoteTrackStatsSink>(SinkName));
+			}
+
+			SsrcSinks[RtpStats._local._ssrc]->Process(RtpStats, AssociatedPlayerId, SecondsDelta);
 		}
 
-		if (InStats._remoteAudioTracks._size > 0)
+		// remote audio stats
+		for (uint64 i = 0; i < InStats._remoteAudioTracks._size; i++)
 		{
-			if (InStats._remoteAudioTracks._size > 1)
+			const EpicRtcRemoteTrackStats& RemoteAudioTrackStats = InStats._remoteAudioTracks._ptr[i];
+
+			// Process audio track rtp stats
+			if (!RemoteAudioTrackSinks.Contains(i))
 			{
-				UE_LOGFMT(LogPixelStreaming2RTC, Warning, "Multiple remote audio tracks is unsupported. Expected: 1, Actual: {0}", InStats._remoteAudioTracks._size);
+				RemoteAudioTrackSinks.Add(i, {});
 			}
-			RemoteAudioTrackStatsSink->Process(InStats._remoteAudioTracks._ptr[0]._rtp, AssociatedPlayerId, SecondsDelta);
+
+			TMap<uint32, TUniquePtr<FRTPRemoteTrackStatsSink>>& SsrcSinks = RemoteVideoTrackSinks[i];
+			const EpicRtcRemoteTrackRtpStats&					  RtpStats = RemoteAudioTrackStats._rtp;
+
+			if (!SsrcSinks.Contains(RtpStats._local._ssrc))
+			{
+				FName SinkName = FName(*FString::Printf(TEXT("%s [%u] (%u)"), *RTCStatCategories::RemoteAudioTrack, i, RtpStats._local._ssrc));
+				SsrcSinks.Add(RtpStats._local._ssrc, MakeUnique<FRTPRemoteTrackStatsSink>(SinkName));
+			}
+
+			SsrcSinks[RtpStats._local._ssrc]->Process(RtpStats, AssociatedPlayerId, SecondsDelta);
 		}
 
-		if (InStats._dataTracks._size > 0)
+		// data track stats
+		for (uint64 i = 0; i < InStats._dataTracks._size; i++)
 		{
-			if (InStats._dataTracks._size > 1)
+			const EpicRtcDataTrackStats& DataTrackStats = InStats._dataTracks._ptr[i];
+
+			// Process data track stats
+			if (!DataTrackSinks.Contains(i))
 			{
-				UE_LOGFMT(LogPixelStreaming2RTC, Warning, "Multiple data tracks is unsupported. Expected: 1, Actual: {0}", InStats._dataTracks._size);
+				FName SinkName = FName(*FString::Printf(TEXT("%s [%u]"), *RTCStatCategories::DataChannel, i));
+				DataTrackSinks.Add(i, MakeUnique<FDataTrackStatsSink>(SinkName));
 			}
-			DataTrackStatsSink->Process(InStats._dataTracks._ptr[0], AssociatedPlayerId, SecondsDelta);
+
+			DataTrackSinks[i]->Process(DataTrackStats, AssociatedPlayerId, SecondsDelta);
 		}
 
+		// transport stats
 		if (InStats._transports._size > 0)
 		{
 			//(Nazar.Rudenko): More than one transport is possible only if we are not using bundle which we do
@@ -120,13 +195,13 @@ namespace UE::PixelStreaming2
 			}
 		}
 
-		LastCalculationCycles = FPlatformTime::Cycles64();
+		LastCalculationCycles = CyclesNow;
 	}
 	/**
 	 * ---------- FRTPLocalVideoTrackSink ----------
 	 */
-	FRTCStatsCollector::FRTPLocalVideoTrackStatsSink::FRTPLocalVideoTrackStatsSink()
-		: FStatsSink(RTCStatCategories::LocalVideoTrack)
+	FRTCStatsCollector::FRTPLocalVideoTrackStatsSink::FRTPLocalVideoTrackStatsSink(FName InCategory)
+		: FStatsSink(InCategory)
 	{
 		// These stats will be extracted from the stat reports and emitted straight to screen
 		Add(PixelStreaming2StatNames::FirCount, 0);
@@ -309,7 +384,7 @@ namespace UE::PixelStreaming2
 		});
 	}
 
-	void FRTCStatsCollector::FRTPLocalVideoTrackStatsSink::Process(const EpicRtcLocalVideoTrackStats& InStats, const FString& PeerId, double SecondsDelta)
+	void FRTCStatsCollector::FRTPLocalVideoTrackStatsSink::Process(const EpicRtcLocalTrackRtpStats& InStats, const FString& PeerId, double SecondsDelta)
 	{
 		FStats* PSStats = FStats::Get();
 		if (!PSStats)
@@ -317,125 +392,112 @@ namespace UE::PixelStreaming2
 			return;
 		}
 
-		if (InStats._rtp._size == 0)
-		{
-			return;
-		}
-
-		// Huge assumption, but assume that normal peers only have 1 video rtp stats and simulcast is anything > 1
-		bool bIsSFU = InStats._rtp._size > 1;
-
-		for (int i = 0; i < InStats._rtp._size; i++)
+		for (TPair<FName, FRTCTrackedStat>& Tuple : Stats)
 		{
-			const EpicRtcLocalTrackRtpStats& RtpStats = InStats._rtp._ptr[i];
-			const FString					 StatsId = bIsSFU ? FString::FromInt(RtpStats._local._ssrc) : PeerId;
-
-			for (TPair<FName, FRTCTrackedStat>& Tuple : Stats)
+			double NewValue = 0;
+			if (Tuple.Key == PixelStreaming2StatNames::FirCount)
 			{
-				double NewValue = 0;
-				if (Tuple.Key == PixelStreaming2StatNames::FirCount)
-				{
-					NewValue = RtpStats._local._firCount;
-				}
-				else if (Tuple.Key == PixelStreaming2StatNames::PliCount)
-				{
-					NewValue = RtpStats._local._pliCount;
-				}
-				else if (Tuple.Key == PixelStreaming2StatNames::NackCount)
-				{
-					NewValue = RtpStats._local._nackCount;
-				}
-				else if (Tuple.Key == PixelStreaming2StatNames::RetransmittedBytesSent)
-				{
-					NewValue = RtpStats._local._retransmittedBytesSent;
-				}
-				else if (Tuple.Key == PixelStreaming2StatNames::TotalEncodeBytesTarget)
-				{
-					NewValue = RtpStats._local._totalEncodedBytesTarget;
-				}
-				else if (Tuple.Key == PixelStreaming2StatNames::KeyFramesEncoded)
-				{
-					NewValue = RtpStats._local._keyFramesEncoded;
-				}
-				else if (Tuple.Key == PixelStreaming2StatNames::FrameWidth)
-				{
-					NewValue = RtpStats._local._frameWidth;
-				}
-				else if (Tuple.Key == PixelStreaming2StatNames::FrameHeight)
-				{
-					NewValue = RtpStats._local._frameHeight;
-				}
-				else if (Tuple.Key == PixelStreaming2StatNames::HugeFramesSent)
-				{
-					NewValue = RtpStats._local._hugeFramesSent;
-				}
-				else if (Tuple.Key == PixelStreaming2StatNames::TotalPacketSendDelay)
-				{
-					NewValue = RtpStats._local._totalPacketSendDelay;
-				}
-				else if (Tuple.Key == PixelStreaming2StatNames::TargetBitrate)
-				{
-					NewValue = RtpStats._local._targetBitrate;
-				}
-				else if (Tuple.Key == PixelStreaming2StatNames::FramesSent)
-				{
-					NewValue = RtpStats._local._framesSent;
-				}
-				else if (Tuple.Key == PixelStreaming2StatNames::FramesReceived)
-				{
-					// TODO(Nazar.Rudenko): Available for inbound tracks only
-				}
-				else if (Tuple.Key == PixelStreaming2StatNames::BytesSent)
-				{
-					NewValue = RtpStats._local._bytesSent;
-				}
-				else if (Tuple.Key == PixelStreaming2StatNames::BytesReceived)
-				{
-					// TODO(Nazar.Rudenko): Available for inbound tracks only
-				}
-				else if (Tuple.Key == PixelStreaming2StatNames::QPSum)
-				{
-					NewValue = RtpStats._local._qpSum;
-				}
-				else if (Tuple.Key == PixelStreaming2StatNames::TotalEncodeTime)
-				{
-					NewValue = RtpStats._local._totalEncodeTime;
-				}
-				else if (Tuple.Key == PixelStreaming2StatNames::FramesEncoded)
-				{
-					NewValue = RtpStats._local._framesEncoded;
-				}
-				else if (Tuple.Key == PixelStreaming2StatNames::FramesDecoded)
-				{
-					// TODO(Nazar.Rudenko): Available for inbound tracks only
-				}
-				else if (Tuple.Key == PixelStreaming2StatNames::PacketsLost)
-				{
-					NewValue = RtpStats._remote._packetsLost;
-				}
-				else if (Tuple.Key == PixelStreaming2StatNames::Jitter)
-				{
-					NewValue = RtpStats._remote._jitter;
-				}
-				else if (Tuple.Key == PixelStreaming2StatNames::RoundTripTime)
-				{
-					NewValue = RtpStats._remote._roundTripTime;
-				}
-
-				if (UpdateValue(NewValue, &Tuple.Value))
-				{
-					PSStats->StorePeerStat(StatsId, RTCStatCategories::LocalVideoTrack, Tuple.Value.GetLatestStat());
-				}
+				NewValue = InStats._local._firCount;
+			}
+			else if (Tuple.Key == PixelStreaming2StatNames::PliCount)
+			{
+				NewValue = InStats._local._pliCount;
+			}
+			else if (Tuple.Key == PixelStreaming2StatNames::NackCount)
+			{
+				NewValue = InStats._local._nackCount;
+			}
+			else if (Tuple.Key == PixelStreaming2StatNames::RetransmittedBytesSent)
+			{
+				NewValue = InStats._local._retransmittedBytesSent;
+			}
+			else if (Tuple.Key == PixelStreaming2StatNames::TotalEncodeBytesTarget)
+			{
+				NewValue = InStats._local._totalEncodedBytesTarget;
+			}
+			else if (Tuple.Key == PixelStreaming2StatNames::KeyFramesEncoded)
+			{
+				NewValue = InStats._local._keyFramesEncoded;
+			}
+			else if (Tuple.Key == PixelStreaming2StatNames::FrameWidth)
+			{
+				NewValue = InStats._local._frameWidth;
+			}
+			else if (Tuple.Key == PixelStreaming2StatNames::FrameHeight)
+			{
+				NewValue = InStats._local._frameHeight;
+			}
+			else if (Tuple.Key == PixelStreaming2StatNames::HugeFramesSent)
+			{
+				NewValue = InStats._local._hugeFramesSent;
+			}
+			else if (Tuple.Key == PixelStreaming2StatNames::TotalPacketSendDelay)
+			{
+				NewValue = InStats._local._totalPacketSendDelay;
+			}
+			else if (Tuple.Key == PixelStreaming2StatNames::TargetBitrate)
+			{
+				NewValue = InStats._local._targetBitrate;
+			}
+			else if (Tuple.Key == PixelStreaming2StatNames::FramesSent)
+			{
+				NewValue = InStats._local._framesSent;
+			}
+			else if (Tuple.Key == PixelStreaming2StatNames::FramesReceived)
+			{
+				// TODO(Nazar.Rudenko): Available for inbound tracks only
+			}
+			else if (Tuple.Key == PixelStreaming2StatNames::BytesSent)
+			{
+				NewValue = InStats._local._bytesSent;
+			}
+			else if (Tuple.Key == PixelStreaming2StatNames::BytesReceived)
+			{
+				// TODO(Nazar.Rudenko): Available for inbound tracks only
+			}
+			else if (Tuple.Key == PixelStreaming2StatNames::QPSum)
+			{
+				NewValue = InStats._local._qpSum;
+			}
+			else if (Tuple.Key == PixelStreaming2StatNames::TotalEncodeTime)
+			{
+				NewValue = InStats._local._totalEncodeTime;
+			}
+			else if (Tuple.Key == PixelStreaming2StatNames::FramesEncoded)
+			{
+				NewValue = InStats._local._framesEncoded;
+			}
+			else if (Tuple.Key == PixelStreaming2StatNames::FramesDecoded)
+			{
+				// TODO(Nazar.Rudenko): Available for inbound tracks only
+			}
+			else if (Tuple.Key == PixelStreaming2StatNames::PacketsLost)
+			{
+				NewValue = InStats._remote._packetsLost;
+			}
+			else if (Tuple.Key == PixelStreaming2StatNames::Jitter)
+			{
+				NewValue = InStats._remote._jitter;
+			}
+			else if (Tuple.Key == PixelStreaming2StatNames::RoundTripTime)
+			{
+				NewValue = InStats._remote._roundTripTime;
 			}
 
-			PostProcess(PSStats, StatsId, SecondsDelta);
+			if (UpdateValue(NewValue, &Tuple.Value))
+			{
+				PSStats->StorePeerStat(PeerId, Category, Tuple.Value.GetLatestStat());
+			}
 		}
+
+		PostProcess(PSStats, PeerId, SecondsDelta);
 	}
+
 	/**
 	 * ---------- FRTPLocalAudioTrackStatsSink ----------
 	 */
-	FRTCStatsCollector::FRTPLocalAudioTrackStatsSink::FRTPLocalAudioTrackStatsSink()
-		: FStatsSink(RTCStatCategories::LocalAudioTrack)
+	FRTCStatsCollector::FRTPLocalAudioTrackStatsSink::FRTPLocalAudioTrackStatsSink(FName InCategory)
+		: FStatsSink(InCategory)
 	{
 		// These stats will be extracted from the stat reports and emitted straight to screen
 		Add(PixelStreaming2StatNames::FirCount, 0);
@@ -656,7 +718,7 @@ namespace UE::PixelStreaming2
 
 			if (UpdateValue(NewValue, &Tuple.Value))
 			{
-				PSStats->StorePeerStat(PeerId, RTCStatCategories::LocalAudioTrack, Tuple.Value.GetLatestStat());
+				PSStats->StorePeerStat(PeerId, Category, Tuple.Value.GetLatestStat());
 			}
 		}
 
@@ -925,7 +987,7 @@ namespace UE::PixelStreaming2
 
 			if (UpdateValue(NewValue, &Tuple.Value))
 			{
-				PSStats->StorePeerStat(PeerId, RTCStatCategories::RemoteVideoTrack, Tuple.Value.GetLatestStat());
+				PSStats->StorePeerStat(PeerId, Category, Tuple.Value.GetLatestStat());
 			}
 		}
 
@@ -935,8 +997,8 @@ namespace UE::PixelStreaming2
 	/**
 	 * ---------- FRTPVideoSourceSink ----------
 	 */
-	FRTCStatsCollector::FVideoSourceStatsSink::FVideoSourceStatsSink()
-		: FStatsSink(RTCStatCategories::VideoSource)
+	FRTCStatsCollector::FVideoSourceStatsSink::FVideoSourceStatsSink(FName InCategory)
+		: FStatsSink(InCategory)
 	{
 		// Track video source fps
 		Add(PixelStreaming2StatNames::SourceFps, 0);
@@ -960,7 +1022,7 @@ namespace UE::PixelStreaming2
 
 			if (UpdateValue(NewValue, &Tuple.Value))
 			{
-				PSStats->StorePeerStat(PeerId, RTCStatCategories::VideoSource, Tuple.Value.GetLatestStat());
+				PSStats->StorePeerStat(PeerId, Category, Tuple.Value.GetLatestStat());
 			}
 		}
 
@@ -970,8 +1032,8 @@ namespace UE::PixelStreaming2
 	/**
 	 * ---------- FRTPAudioSourceSink ----------
 	 */
-	FRTCStatsCollector::FAudioSourceStatsSink::FAudioSourceStatsSink()
-		: FStatsSink(RTCStatCategories::AudioSource)
+	FRTCStatsCollector::FAudioSourceStatsSink::FAudioSourceStatsSink(FName InCategory)
+		: FStatsSink(InCategory)
 	{
 		Add(PixelStreaming2StatNames::AudioLevel, 0);
 		Add(PixelStreaming2StatNames::TotalSamplesDuration, 0);
@@ -999,7 +1061,7 @@ namespace UE::PixelStreaming2
 
 			if (UpdateValue(NewValue, &Tuple.Value))
 			{
-				PSStats->StorePeerStat(PeerId, RTCStatCategories::AudioSource, Tuple.Value.GetLatestStat());
+				PSStats->StorePeerStat(PeerId, Category, Tuple.Value.GetLatestStat());
 			}
 		}
 
@@ -1009,8 +1071,8 @@ namespace UE::PixelStreaming2
 	/**
 	 * ----------- FDataChannelSink -----------
 	 */
-	FRTCStatsCollector::FDataTrackStatsSink::FDataTrackStatsSink()
-		: FStatsSink(RTCStatCategories::DataChannel)
+	FRTCStatsCollector::FDataTrackStatsSink::FDataTrackStatsSink(FName InCategory)
+		: FStatsSink(InCategory)
 	{
 		// These names are added as aliased names because `bytesSent` is ambiguous stat that is used across inbound-rtp, outbound-rtp, and data-channel
 		// so to disambiguate which state we are referring to we record the `bytesSent` stat for the data-channel but store and report it as `data-channel-bytesSent`
@@ -1050,7 +1112,7 @@ namespace UE::PixelStreaming2
 
 			if (UpdateValue(NewValue, &Tuple.Value))
 			{
-				PSStats->StorePeerStat(PeerId, RTCStatCategories::DataChannel, Tuple.Value.GetLatestStat());
+				PSStats->StorePeerStat(PeerId, Category, Tuple.Value.GetLatestStat());
 			}
 		}
 
@@ -1060,8 +1122,8 @@ namespace UE::PixelStreaming2
 	/**
 	 * ---------- FRTPAudioSourceSink ----------
 	 */
-	FRTCStatsCollector::FCandidatePairStatsSink::FCandidatePairStatsSink()
-		: FStatsSink(RTCStatCategories::CandidatePair)
+	FRTCStatsCollector::FCandidatePairStatsSink::FCandidatePairStatsSink(FName InCategory)
+		: FStatsSink(InCategory)
 	{
 		Add(PixelStreaming2StatNames::AvailableOutgoingBitrate, 0);
 		Add(PixelStreaming2StatNames::AvailableIncomingBitrate, 0);
diff --git a/Engine/Plugins/Media/PixelStreaming2/Source/PixelStreaming2RTC/Private/RTCStatsCollector.h b/Engine/Plugins/Media/PixelStreaming2/Source/PixelStreaming2RTC/Private/RTCStatsCollector.h
--- a/Engine/Plugins/Media/PixelStreaming2/Source/PixelStreaming2RTC/Private/RTCStatsCollector.h
+++ b/Engine/Plugins/Media/PixelStreaming2/Source/PixelStreaming2RTC/Private/RTCStatsCollector.h
@@ -10,14 +10,14 @@ namespace UE::PixelStreaming2
 {
 	namespace RTCStatCategories
 	{
-		const FName LocalVideoTrack = FName(TEXT("outbound-video-track"));
-		const FName LocalAudioTrack = FName(TEXT("outbound-audio-track"));
-		const FName VideoSource = FName(TEXT("video-source"));
-		const FName AudioSource = FName(TEXT("audio-source"));
-		const FName DataChannel = FName(TEXT("data-channel"));
-		const FName RemoteVideoTrack = FName(TEXT("inbound-video-track"));
-		const FName RemoteAudioTrack = FName(TEXT("inbound-audio-track"));
-		const FName CandidatePair = FName(TEXT("candidate-pair"));
+		const FString LocalVideoTrack = TEXT("video-track-outbound");
+		const FString LocalAudioTrack = TEXT("audio-track-outbound");
+		const FString VideoSource = TEXT("video-source");
+		const FString AudioSource = TEXT("audio-source");
+		const FString DataChannel = TEXT("data-channel");
+		const FString RemoteVideoTrack = TEXT("video-track-inbound");
+		const FString RemoteAudioTrack = TEXT("audio-track-inbound");
+		const FString CandidatePair = TEXT("candidate-pair");
 	} // namespace RTCStatCategories
 
 	class FRTCStatsCollector
@@ -109,15 +109,15 @@ namespace UE::PixelStreaming2
 		class FRTPLocalVideoTrackStatsSink : public FStatsSink
 		{
 		public:
-			FRTPLocalVideoTrackStatsSink();
+			FRTPLocalVideoTrackStatsSink(FName Category);
 			virtual ~FRTPLocalVideoTrackStatsSink() = default;
-			void Process(const EpicRtcLocalVideoTrackStats& InStats, const FString& PeerId, double SecondsDelta);
+			void Process(const EpicRtcLocalTrackRtpStats& InStats, const FString& PeerId, double SecondsDelta);
 		};
 
 		class FRTPLocalAudioTrackStatsSink : public FStatsSink
 		{
 		public:
-			FRTPLocalAudioTrackStatsSink();
+			FRTPLocalAudioTrackStatsSink(FName Category);
 			virtual ~FRTPLocalAudioTrackStatsSink() = default;
 			void Process(const EpicRtcLocalTrackRtpStats& InStats, const FString& PeerId, double SecondsDelta);
 		};
@@ -133,7 +133,7 @@ namespace UE::PixelStreaming2
 		class FVideoSourceStatsSink : public FStatsSink
 		{
 		public:
-			FVideoSourceStatsSink();
+			FVideoSourceStatsSink(FName Category);
 			virtual ~FVideoSourceStatsSink() = default;
 			void Process(const EpicRtcVideoSourceStats& InStats, const FString& PeerId, double SecondsDelta);
 		};
@@ -141,7 +141,7 @@ namespace UE::PixelStreaming2
 		class FAudioSourceStatsSink : public FStatsSink
 		{
 		public:
-			FAudioSourceStatsSink();
+			FAudioSourceStatsSink(FName Category);
 			virtual ~FAudioSourceStatsSink() = default;
 			void Process(const EpicRtcAudioSourceStats& InStats, const FString& PeerId, double SecondsDelta);
 		};
@@ -152,7 +152,7 @@ namespace UE::PixelStreaming2
 		class FDataTrackStatsSink : public FStatsSink
 		{
 		public:
-			FDataTrackStatsSink();
+			FDataTrackStatsSink(FName Category);
 			virtual ~FDataTrackStatsSink() = default;
 
 			void Process(const EpicRtcDataTrackStats& InStats, const FString& PeerId, double SecondsDelta);
@@ -161,27 +161,28 @@ namespace UE::PixelStreaming2
 		class FCandidatePairStatsSink : public FStatsSink
 		{
 		public:
-			FCandidatePairStatsSink();
+			FCandidatePairStatsSink(FName Category);
 			virtual ~FCandidatePairStatsSink() = default;
 			void Process(const EpicRtcIceCandidatePairStats& InStats, const FString& PeerId, double SecondsDelta);
 		};
 
 	private:
-		FString									 AssociatedPlayerId;
-		uint64									 LastCalculationCycles;
-		bool									 bIsEnabled;
+		FString AssociatedPlayerId;
+		uint64	LastCalculationCycles;
+		bool	bIsEnabled;
+
+		//    index,         ssrc                          sink
+		TMap<uint64, TMap<uint32, TUniquePtr<FRTPLocalVideoTrackStatsSink>>> LocalVideoTrackSinks;
+		TMap<uint64, TMap<uint32, TUniquePtr<FRTPLocalAudioTrackStatsSink>>> LocalAudioTrackSinks;
 
-		TUniquePtr<FRTPLocalVideoTrackStatsSink> LocalVideoTrackStatsSink;
-		TUniquePtr<FRTPLocalAudioTrackStatsSink> LocalAudioTrackStatsSink;
+		TMap<uint64, TMap<uint32, TUniquePtr<FRTPRemoteTrackStatsSink>>> RemoteVideoTrackSinks;
+		TMap<uint64, TMap<uint32, TUniquePtr<FRTPRemoteTrackStatsSink>>> RemoteAudioTrackSinks;
 
-		TUniquePtr<FVideoSourceStatsSink>		 VideoSourceStatsSink;
-		TUniquePtr<FAudioSourceStatsSink>	 AudioSourceStatsSink;
+		TMap<uint64, TUniquePtr<FVideoSourceStatsSink>> VideoSourceSinks;
+		TMap<uint64, TUniquePtr<FAudioSourceStatsSink>> AudioSourceSinks;
 
-		TUniquePtr<FRTPRemoteTrackStatsSink> RemoteVideoTrackStatsSink;
-		TUniquePtr<FRTPRemoteTrackStatsSink> RemoteAudioTrackStatsSink;
+		TMap<uint64, TUniquePtr<FDataTrackStatsSink>> DataTrackSinks;
 
-		TUniquePtr<FDataTrackStatsSink>		 DataTrackStatsSink;
-		
-		TUniquePtr<FCandidatePairStatsSink>	 CandidatePairStatsSink;
+		TUniquePtr<FCandidatePairStatsSink> CandidatePairStatsSink;
 	};
 } // namespace UE::PixelStreaming2
\ No newline at end of file
diff --git a/Engine/Plugins/Media/PixelStreaming2/Source/PixelStreaming2RTC/Private/Stats.cpp b/Engine/Plugins/Media/PixelStreaming2/Source/PixelStreaming2RTC/Private/Stats.cpp
--- a/Engine/Plugins/Media/PixelStreaming2/Source/PixelStreaming2RTC/Private/Stats.cpp
+++ b/Engine/Plugins/Media/PixelStreaming2/Source/PixelStreaming2RTC/Private/Stats.cpp
@@ -23,7 +23,7 @@ IPixelStreaming2Stats& IPixelStreaming2Stats::Get()
 static FCanvasTextItem CreateText(const FString& String, double X, double Y)
 {
 	FText			TextToDisplay = FText::FromString(String);
-	FCanvasTextItem Text(FVector2D(X, Y), TextToDisplay, FSlateFontInfo(FSlateFontInfo(UEngine::GetSmallFont(), 12)), FLinearColor(0, 1, 0));
+	FCanvasTextItem Text(FVector2D(X, Y), TextToDisplay, FSlateFontInfo(FSlateFontInfo(UEngine::GetSmallFont(), 10)), FLinearColor(0, 1, 0));
 	Text.EnableShadow(FLinearColor::Black);
 	return Text;
 }
@@ -79,13 +79,35 @@ namespace UE::PixelStreaming2
 		}
 	}
 
-	bool FStats::QueryPeerStat(const FString& PlayerId, FName StatCategory, FName StatToQuery, double& OutValue) const
+	bool FStats::QueryPeerStat(const FString& PlayerId, FName InStatCategory, FName StatToQuery, double& OutValue) const
 	{
 		FScopeLock Lock(&PeerStatsCS);
 
 		if (const FPeerStats* SinglePeerStats = PeerStats.Find(PlayerId))
 		{
-			return SinglePeerStats->GetStat(StatCategory, StatToQuery, OutValue);
+			const TMap<FName, FStatGroup>& StatGroups = SinglePeerStats->GetStatGroups();
+			TArray<FName>			 StatCategories;
+			StatGroups.GetKeys(StatCategories);
+
+			// Stat groups contain a name as well as additional info like track index and ssrc
+			// When querying a stat we need to find all matching stats based on name
+			TArray<FName> MatchedStatCategories;
+			for (FName& StatCategory : StatCategories)
+			{
+				if (StatCategory.ToString().Contains(InStatCategory.ToString()))
+				{
+					MatchedStatCategories.Add(StatCategory);
+				}
+			}
+
+			if (MatchedStatCategories.Num() == 0)
+			{
+				return false;
+			}
+
+			// TODO (william.belcher): This is lazy and only queries the first matched category but since
+			// this code is only used on the p2p use case where there only is one matching category it's fine
+			return SinglePeerStats->GetStat(MatchedStatCategories[0], StatToQuery, OutValue);
 		}
 
 		return false;
@@ -549,6 +571,9 @@ namespace UE::PixelStreaming2
 		if (!StatGroups.Contains(StatCategory))
 		{
 			StatGroups.Add(StatCategory, FStatGroup(StatCategory));
+			StatGroups.KeySort([](FName A, FName B) {
+				return A.ToString().Compare(B.ToString(), ESearchCase::IgnoreCase) < 0;
+			});
 		}
 		return StatGroups[StatCategory].StoreStat(StatToStore);
 	}
diff --git a/Engine/Plugins/Media/PixelStreaming2/Source/PixelStreaming2RTC/Private/Stats.h b/Engine/Plugins/Media/PixelStreaming2/Source/PixelStreaming2RTC/Private/Stats.h
--- a/Engine/Plugins/Media/PixelStreaming2/Source/PixelStreaming2RTC/Private/Stats.h
+++ b/Engine/Plugins/Media/PixelStreaming2/Source/PixelStreaming2RTC/Private/Stats.h
@@ -80,7 +80,7 @@ namespace UE::PixelStreaming2
 	public:
 		FStatGroup(FName CategoryName)
 			: GroupName(CategoryName)
-			, CategoryCanvasItem(FVector2D(0, 0), FText::FromString(FString::Printf(TEXT("---%s---"), *CategoryName.ToString())), FSlateFontInfo(FSlateFontInfo(UEngine::GetSmallFont(), 15)), FLinearColor(0, 0.9, 0.1))
+			, CategoryCanvasItem(FVector2D(0, 0), FText::FromString(FString::Printf(TEXT("---%s---"), *CategoryName.ToString())), FSlateFontInfo(FSlateFontInfo(UEngine::GetSmallFont(), 12)), FLinearColor(0, 0.9, 0.1))
 		{
 			CategoryCanvasItem.EnableShadow(FLinearColor::Black);
 		}
