[PS2] Add: Experimental Play-In-Editor streaming support

diff --git a/Engine/Plugins/Media/PixelStreaming2/Source/PixelStreaming2Editor/Private/PixelStreaming2EditorModule.cpp b/Engine/Plugins/Media/PixelStreaming2/Source/PixelStreaming2Editor/Private/PixelStreaming2EditorModule.cpp
--- a/Engine/Plugins/Media/PixelStreaming2/Source/PixelStreaming2Editor/Private/PixelStreaming2EditorModule.cpp
+++ b/Engine/Plugins/Media/PixelStreaming2/Source/PixelStreaming2Editor/Private/PixelStreaming2EditorModule.cpp
@@ -12,6 +12,7 @@
 #include "IPixelStreaming2RTCModule.h"
 #include "IPixelStreaming2Streamer.h"
 #include "VideoProducerBackBufferComposited.h"
+#include "VideoProducerPIEViewport.h"
 #include "VideoProducerViewport.h"
 #include "PixelStreaming2PluginSettings.h"
 #include "UnrealEngine.h"
@@ -186,6 +187,11 @@ namespace UE::EditorPixelStreaming2
 				StartStreaming(Source);
 			}
 		});
+
+#if WITH_EDITOR
+		FEditorDelegates::PostPIEStarted.AddRaw(this, &FPixelStreaming2EditorModule::OnBeginPIE);
+		FEditorDelegates::EndPIE.AddRaw(this, &FPixelStreaming2EditorModule::OnEndPIE);
+#endif
 	}
 
 	void FPixelStreaming2EditorModule::StartStreaming(EPixelStreaming2EditorStreamTypes InStreamType)
@@ -467,6 +473,55 @@ namespace UE::EditorPixelStreaming2
 			FSlateNotificationManager::Get().QueueNotification(Info);
 		}
 	}
+
+#if WITH_EDITOR
+	void FPixelStreaming2EditorModule::OnBeginPIE(bool bIsSimulating)
+	{
+		IPixelStreaming2Module& Module = IPixelStreaming2Module::Get();
+		PIEStreamer = Module.CreateStreamer(Module.GetDefaultStreamerID());
+		// Give the PIE streamer the default url if the user hasn't specified one when launching the editor
+		if (PIEStreamer->GetConnectionURL().IsEmpty())
+		{
+			// No URL was passed on the command line, initialize defaults
+			PIEStreamer->SetConnectionURL(FString::Printf(TEXT("%s:%d"), *SignallingDomain, StreamerPort));
+		}
+
+		PIEStreamer->SetVideoProducer(FVideoProducerPIEViewport::Create());
+
+		TSharedPtr<IPixelStreaming2InputHandler> InputHandler = PIEStreamer->GetInputHandler().Pin();
+		if (!InputHandler.IsValid())
+		{
+			return;
+		}
+
+		FWorldContext* PIEWorldContext = GEditor->GetPIEWorldContext();
+		if (!PIEWorldContext)
+		{
+			UE_LOGFMT(LogPixelStreaming2Editor, Warning, "Failed to get editor PIE world context. PIE streamer will not stream!");
+			return;
+		}
+
+		UWorld* World = PIEWorldContext->World();
+		if (!World)
+		{
+			UE_LOGFMT(LogPixelStreaming2Editor, Warning, "Failed to get editor PIE world. PIE streamer will not stream!");
+			return;
+		}
+
+		UGameViewportClient* Viewport = World->GetGameViewport();
+		InputHandler->SetTargetViewport(Viewport->GetGameViewportWidget());
+		InputHandler->SetTargetWindow(Viewport->GetWindow());
+		InputHandler->SetInputType(EPixelStreaming2InputType::RouteToWindow);
+		PIEStreamer->StartStreaming();
+	}
+
+	void FPixelStreaming2EditorModule::OnEndPIE(bool bIsSimulating)
+	{
+		PIEStreamer->StopStreaming();
+
+		IPixelStreaming2Module::Get().DeleteStreamer(PIEStreamer);
+	}
+#endif
 } // namespace UE::EditorPixelStreaming2
 
 #undef IMAGE_BRUSH_SVG
diff --git a/Engine/Plugins/Media/PixelStreaming2/Source/PixelStreaming2Editor/Private/PixelStreaming2EditorModule.h b/Engine/Plugins/Media/PixelStreaming2/Source/PixelStreaming2Editor/Private/PixelStreaming2EditorModule.h
--- a/Engine/Plugins/Media/PixelStreaming2/Source/PixelStreaming2Editor/Private/PixelStreaming2EditorModule.h
+++ b/Engine/Plugins/Media/PixelStreaming2/Source/PixelStreaming2Editor/Private/PixelStreaming2EditorModule.h
@@ -20,8 +20,8 @@ namespace UE::EditorPixelStreaming2
 		virtual void StartStreaming(EPixelStreaming2EditorStreamTypes InStreamType) override;
 		virtual void StopStreaming() override;
 
-		virtual void												  StartSignalling() override;
-		virtual void												  StopSignalling() override;
+		virtual void											StartSignalling() override;
+		virtual void											StopSignalling() override;
 		virtual TSharedPtr<UE::PixelStreaming2Servers::IServer> GetSignallingServer() override;
 
 		virtual void	SetSignallingDomain(const FString& InSignallingDomain) override;
@@ -38,6 +38,11 @@ namespace UE::EditorPixelStreaming2
 		void DisableCPUThrottlingSetting();
 		void RestoreCPUThrottlingSetting();
 
+#if WITH_EDITOR
+		void OnBeginPIE(bool bIsSimulating);
+		void OnEndPIE(bool bIsSimulating);
+#endif
+
 		TSharedPtr<FPixelStreaming2Toolbar> Toolbar;
 		// Signalling/webserver
 		TSharedPtr<UE::PixelStreaming2Servers::IServer> SignallingServer;
@@ -56,6 +61,10 @@ namespace UE::EditorPixelStreaming2
 		// The streamer used by the PixelStreaming2Editor module
 		TSharedPtr<IPixelStreaming2Streamer> EditorStreamer;
 
+#if WITH_EDITOR
+		TSharedPtr<IPixelStreaming2Streamer> PIEStreamer;
+#endif
+
 		bool bOldCPUThrottlingSetting = false;
 	};
 } // namespace UE::EditorPixelStreaming2
