[PS2] Fix: Application crash when processing a gamepad analog event if the gamepad has been disconnected

diff --git a/Engine/Plugins/Media/PixelStreaming2/Source/PixelStreaming2Input/Private/InputDevice.cpp b/Engine/Plugins/Media/PixelStreaming2/Source/PixelStreaming2Input/Private/InputDevice.cpp
--- a/Engine/Plugins/Media/PixelStreaming2/Source/PixelStreaming2Input/Private/InputDevice.cpp
+++ b/Engine/Plugins/Media/PixelStreaming2/Source/PixelStreaming2Input/Private/InputDevice.cpp
@@ -111,8 +111,8 @@ namespace UE::PixelStreaming2Input
 
 	void FInputDevice::OnControllerDisconnected(uint8 DeleteControllerId)
 	{
-		FInputDeviceId							DeviceId = INPUTDEVICEID_NONE;
-		FPlatformUserId							PlatformUserId = PLATFORMUSERID_NONE;
+		FInputDeviceId	DeviceId = INPUTDEVICEID_NONE;
+		FPlatformUserId PlatformUserId = PLATFORMUSERID_NONE;
 		if (!GetPlatformUserAndDevice(DeleteControllerId, DeviceId, PlatformUserId))
 		{
 			return;
@@ -137,4 +137,17 @@ namespace UE::PixelStreaming2Input
 		OutPlatformUser = DevicePair.Get<1>();
 		return true;
 	}
+
+	bool FInputDevice::GetControllerIdFromDeviceId(FInputDeviceId DeviceId, uint8& OutControllerId)
+	{
+		for (TPair<uint8, TPair<FInputDeviceId, FPlatformUserId>> InputDevicePair : InputDevices)
+		{
+			if (DeviceId == InputDevicePair.Value.Get<0>())
+			{
+				OutControllerId = InputDevicePair.Key;
+				return true;
+			}
+		}
+		return false;
+	}
 } // namespace UE::PixelStreaming2Input
diff --git a/Engine/Plugins/Media/PixelStreaming2/Source/PixelStreaming2Input/Private/InputDevice.h b/Engine/Plugins/Media/PixelStreaming2/Source/PixelStreaming2Input/Private/InputDevice.h
--- a/Engine/Plugins/Media/PixelStreaming2/Source/PixelStreaming2Input/Private/InputDevice.h
+++ b/Engine/Plugins/Media/PixelStreaming2/Source/PixelStreaming2Input/Private/InputDevice.h
@@ -31,6 +31,7 @@ namespace UE::PixelStreaming2Input
 		uint8 OnControllerConnected();
 		void  OnControllerDisconnected(uint8 DeleteControllerId);
 		bool  GetPlatformUserAndDevice(uint8 ControllerId, FInputDeviceId& OutDeviceId, FPlatformUserId& OutPlatformUser);
+		bool  GetControllerIdFromDeviceId(FInputDeviceId DeviceId, uint8& OutControllerId);
 
 		static TSharedPtr<FInputDevice> GetInputDevice();
 
diff --git a/Engine/Plugins/Media/PixelStreaming2/Source/PixelStreaming2Input/Private/InputHandler.cpp b/Engine/Plugins/Media/PixelStreaming2/Source/PixelStreaming2Input/Private/InputHandler.cpp
--- a/Engine/Plugins/Media/PixelStreaming2/Source/PixelStreaming2Input/Private/InputHandler.cpp
+++ b/Engine/Plugins/Media/PixelStreaming2/Source/PixelStreaming2Input/Private/InputHandler.cpp
@@ -1386,6 +1386,21 @@ namespace UE::PixelStreaming2Input
 					return;
 				}
 
+				// Check if this gamepad event is specific to a normal gamepad and not an xr gamepad
+				TArray<FKey> GamepadKeys;
+				FInputConverter::GamepadInputToFKey.GenerateValueArray(GamepadKeys);
+				if (GamepadKeys.Find(*Key) != INDEX_NONE)
+				{
+					// This is a gamepad key. We need to check that the gamepad hasn't been disconnected for we've been able to process this analog event
+					uint8 ControllerIndex;
+					if (!FInputDevice::GetInputDevice()->GetControllerIdFromDeviceId(ControllerId, ControllerIndex))
+					{
+						// We're unable to get a platform user and device for this controller id. That means the controllers been disconnected
+						FKeyIt.RemoveCurrent();
+						continue;
+					}
+				}
+
 				// Pass an analog input along the engine's input processing system
 				bool bHandled = MessageHandler->OnControllerAnalog(Key->GetFName(), PLATFORMUSERID_NONE /* Not used */, ControllerId, AnalogValue.Value);
 				UE_LOG(LogPixelStreaming2Input, VeryVerbose, TEXT("TICKED ANALOG Input: ControllerId = %d; KeyName = %s; IsRepeat = %s; AnalogValue = %.4f; Handled = %s; [Queued for Tick()]"), ControllerId.GetId(), *Key->ToString(), bIsRepeat ? TEXT("True") : TEXT("False"), AnalogValue.Value, bHandled ? TEXT("True") : TEXT("False"));
