[PS2] Fix: Incorrect sink reference

diff --git a/Engine/Plugins/Media/PixelStreaming2/Source/PixelStreaming2RTC/Private/RTCStatsCollector.cpp b/Engine/Plugins/Media/PixelStreaming2/Source/PixelStreaming2RTC/Private/RTCStatsCollector.cpp
--- a/Engine/Plugins/Media/PixelStreaming2/Source/PixelStreaming2RTC/Private/RTCStatsCollector.cpp
+++ b/Engine/Plugins/Media/PixelStreaming2/Source/PixelStreaming2RTC/Private/RTCStatsCollector.cpp
@@ -152,7 +152,7 @@ namespace UE::PixelStreaming2
 				RemoteAudioTrackSinks.Add(i, {});
 			}
 
-			TMap<uint32, TUniquePtr<FRTPRemoteTrackStatsSink>>& SsrcSinks = RemoteVideoTrackSinks[i];
+			TMap<uint32, TUniquePtr<FRTPRemoteTrackStatsSink>>& SsrcSinks = RemoteAudioTrackSinks[i];
 			const EpicRtcRemoteTrackRtpStats&					  RtpStats = RemoteAudioTrackStats._rtp;
 
 			if (!SsrcSinks.Contains(RtpStats._local._ssrc))
