[PS2] Fix: Frontend stream options not being set if previously set by command line

diff --git a/Engine/Plugins/Media/PixelStreaming2/Source/PixelStreaming2/Private/PixelStreaming2Module.cpp b/Engine/Plugins/Media/PixelStreaming2/Source/PixelStreaming2/Private/PixelStreaming2Module.cpp
--- a/Engine/Plugins/Media/PixelStreaming2/Source/PixelStreaming2/Private/PixelStreaming2Module.cpp
+++ b/Engine/Plugins/Media/PixelStreaming2/Source/PixelStreaming2/Private/PixelStreaming2Module.cpp
@@ -72,7 +72,7 @@ namespace UE::PixelStreaming2
 			{
 				UE_LOG(LogPixelStreaming2, Warning, TEXT("Could not setup hardware encoder. This is usually a driver issue or hardware limitation, try reinstalling your drivers."));
 				UE_LOG(LogPixelStreaming2, Warning, TEXT("Falling back to VP8 software video encoding."));
-				UPixelStreaming2PluginSettings::CVarEncoderCodec.AsVariable()->Set(*UE::PixelStreaming2::GetCVarStringFromEnum(EVideoCodec::VP8), ECVF_SetByCommandline);
+				UPixelStreaming2PluginSettings::CVarEncoderCodec.AsVariable()->SetWithCurrentPriority(*UE::PixelStreaming2::GetCVarStringFromEnum(EVideoCodec::VP8));
 				if (UPixelStreaming2Delegates* Delegates = UPixelStreaming2Delegates::Get())
 				{
 					Delegates->OnFallbackToSoftwareEncoding.Broadcast();
@@ -280,31 +280,31 @@ namespace UE::PixelStreaming2
 			// Set Encoder.MinQP Legacy CVar
 			InputHandler->SetCommandHandler(TEXT("Encoder.MinQP"), [](FString PlayerId, FString Descriptor, FString MinQPString) {
 				int MinQP = FCString::Atoi(*MinQPString);
-				UPixelStreaming2PluginSettings::CVarEncoderMaxQuality->Set(FMath::RoundToFloat(100.0f * (1.0f - (FMath::Clamp<int32>(MinQP, 0, 51) / 51.0f))), ECVF_SetByCommandline);
+				UPixelStreaming2PluginSettings::CVarEncoderMaxQuality->SetWithCurrentPriority(FMath::RoundToFloat(100.0f * (1.0f - (FMath::Clamp<int32>(MinQP, 0, 51) / 51.0f))));
 			});
 
 			// Set Encoder.MaxQP Legacy CVar
 			InputHandler->SetCommandHandler(TEXT("Encoder.MaxQP"), [](FString PlayerId, FString Descriptor, FString MaxQPString) {
 				int MaxQP = FCString::Atoi(*MaxQPString);
-				UPixelStreaming2PluginSettings::CVarEncoderMinQuality->Set(FMath::RoundToFloat(100.0f * (1.0f - (FMath::Clamp<int32>(MaxQP, 0, 51) / 51.0f))), ECVF_SetByCommandline);
+				UPixelStreaming2PluginSettings::CVarEncoderMinQuality->SetWithCurrentPriority(FMath::RoundToFloat(100.0f * (1.0f - (FMath::Clamp<int32>(MaxQP, 0, 51) / 51.0f))));
 			});
 
 			// Set Encoder.MinQuality CVar
 			InputHandler->SetCommandHandler(TEXT("Encoder.MinQuality"), [](FString PlayerId, FString Descriptor, FString MinQualityString) {
 				int MinQuality = FCString::Atoi(*MinQualityString);
-				UPixelStreaming2PluginSettings::CVarEncoderMinQuality->Set(FMath::Clamp<int32>(MinQuality, 0, 100), ECVF_SetByCommandline);
+				UPixelStreaming2PluginSettings::CVarEncoderMinQuality->SetWithCurrentPriority(FMath::Clamp<int32>(MinQuality, 0, 100));
 			});
 
 			// Set Encoder.MaxQuality CVar
 			InputHandler->SetCommandHandler(TEXT("Encoder.MaxQuality"), [](FString PlayerId, FString Descriptor, FString MaxQualityString) {
 				int MaxQuality = FCString::Atoi(*MaxQualityString);
-				UPixelStreaming2PluginSettings::CVarEncoderMaxQuality->Set(FMath::Clamp<int32>(MaxQuality, 0, 100), ECVF_SetByCommandline);
+				UPixelStreaming2PluginSettings::CVarEncoderMaxQuality->SetWithCurrentPriority(FMath::Clamp<int32>(MaxQuality, 0, 100));
 			});
 
 			// Set WebRTC max FPS
 			InputHandler->SetCommandHandler(TEXT("WebRTC.Fps"), [](FString PlayerId, FString Descriptor, FString FPSString) {
 				int FPS = FCString::Atoi(*FPSString);
-				UPixelStreaming2PluginSettings::CVarWebRTCFps->Set(FPS, ECVF_SetByCommandline);
+				UPixelStreaming2PluginSettings::CVarWebRTCFps->SetWithCurrentPriority(FPS);
 			});
 
 			// Set MinBitrate
@@ -312,7 +312,7 @@ namespace UE::PixelStreaming2
 				if (InputHandler->IsElevated(PlayerId))
 				{
 					int MinBitrate = FCString::Atoi(*MinBitrateString);
-					UPixelStreaming2PluginSettings::CVarWebRTCMinBitrate->Set(MinBitrate, ECVF_SetByCommandline);
+					UPixelStreaming2PluginSettings::CVarWebRTCMinBitrate->SetWithCurrentPriority(MinBitrate);
 				}
 			});
 
@@ -321,7 +321,7 @@ namespace UE::PixelStreaming2
 				if (InputHandler->IsElevated(PlayerId))
 				{
 					int MaxBitrate = FCString::Atoi(*MaxBitrateString);
-					UPixelStreaming2PluginSettings::CVarWebRTCMaxBitrate->Set(MaxBitrate, ECVF_SetByCommandline);
+					UPixelStreaming2PluginSettings::CVarWebRTCMaxBitrate->SetWithCurrentPriority(MaxBitrate);
 				}
 			});
 
diff --git a/Engine/Plugins/Media/PixelStreaming2/Source/PixelStreaming2Editor/Private/PixelStreaming2Toolbar.cpp b/Engine/Plugins/Media/PixelStreaming2/Source/PixelStreaming2Editor/Private/PixelStreaming2Toolbar.cpp
--- a/Engine/Plugins/Media/PixelStreaming2/Source/PixelStreaming2Editor/Private/PixelStreaming2Toolbar.cpp
+++ b/Engine/Plugins/Media/PixelStreaming2/Source/PixelStreaming2Editor/Private/PixelStreaming2Toolbar.cpp
@@ -45,7 +45,7 @@ DEFINE_LOG_CATEGORY(LogPixelStreaming2Toolbar);
 
 void SetCodec(EVideoCodec Codec)
 {
-	UPixelStreaming2PluginSettings::CVarEncoderCodec.AsVariable()->Set(*UE::PixelStreaming2::GetCVarStringFromEnum(Codec));
+	UPixelStreaming2PluginSettings::CVarEncoderCodec.AsVariable()->SetWithCurrentPriority(*UE::PixelStreaming2::GetCVarStringFromEnum(Codec));
 }
 
 EVideoCodec GetCodec()
@@ -55,7 +55,7 @@ EVideoCodec GetCodec()
 
 void SetUseRemoteSignallingServer(bool UseRemoteSignallingServer)
 {
-	UPixelStreaming2PluginSettings::CVarEditorUseRemoteSignallingServer.AsVariable()->Set(UseRemoteSignallingServer);
+	UPixelStreaming2PluginSettings::CVarEditorUseRemoteSignallingServer.AsVariable()->SetWithCurrentPriority(UseRemoteSignallingServer);
 }
 
 bool GetUseRemoteSignallingServer()
diff --git a/Engine/Plugins/Media/PixelStreaming2/Source/PixelStreaming2RTC/Private/EpicRtcStreamer.cpp b/Engine/Plugins/Media/PixelStreaming2/Source/PixelStreaming2RTC/Private/EpicRtcStreamer.cpp
--- a/Engine/Plugins/Media/PixelStreaming2/Source/PixelStreaming2RTC/Private/EpicRtcStreamer.cpp
+++ b/Engine/Plugins/Media/PixelStreaming2/Source/PixelStreaming2RTC/Private/EpicRtcStreamer.cpp
@@ -457,8 +457,8 @@ namespace UE::PixelStreaming2
 
 	void FEpicRtcStreamer::PlayerRequestsBitrate(FString PlayerId, int MinBitrate, int MaxBitrate)
 	{
-		UPixelStreaming2PluginSettings::CVarWebRTCMinBitrate.AsVariable()->Set(MinBitrate);
-		UPixelStreaming2PluginSettings::CVarWebRTCMaxBitrate.AsVariable()->Set(MaxBitrate);
+		UPixelStreaming2PluginSettings::CVarWebRTCMinBitrate.AsVariable()->SetWithCurrentPriority(MinBitrate);
+		UPixelStreaming2PluginSettings::CVarWebRTCMaxBitrate.AsVariable()->SetWithCurrentPriority(MaxBitrate);
 	}
 
 	void FEpicRtcStreamer::RefreshStreamBitrate()
diff --git a/Engine/Plugins/Media/PixelStreaming2/Source/PixelStreaming2Settings/Private/PixelStreaming2PluginSettings.cpp b/Engine/Plugins/Media/PixelStreaming2/Source/PixelStreaming2Settings/Private/PixelStreaming2PluginSettings.cpp
--- a/Engine/Plugins/Media/PixelStreaming2/Source/PixelStreaming2Settings/Private/PixelStreaming2PluginSettings.cpp
+++ b/Engine/Plugins/Media/PixelStreaming2/Source/PixelStreaming2Settings/Private/PixelStreaming2PluginSettings.cpp
@@ -88,7 +88,7 @@ namespace
 			}
 		}
 
-		CVar->Set(*FString::Join(ValidStringArray, TEXT(",")), ECVF_SetByCommandline);
+		CVar->Set(*FString::Join(ValidStringArray, TEXT(",")), ECVF_SetByConsole);
 	}
 
 	FString ConsoleVariableToCommandArgValue(const FString InCVarName)
@@ -750,7 +750,7 @@ void UPixelStreaming2PluginSettings::PostEditChangeProperty(FPropertyChangedEven
 		// If the property has MappedConsoleVariable metadata fetch it and set it to the enums string value
 		if (IConsoleVariable* ConsoleVariable = IConsoleManager::Get().FindConsoleVariable(*PropertyChangedEvent.Property->GetMetaData(PixelStreaming2MappedConsoleVariableFName)))
 		{
-			ConsoleVariable->Set(*ConsoleString, ECVF_SetByConsole);
+			ConsoleVariable->Set(*ConsoleString, ECVF_SetByProjectSetting);
 		}
 	}
 	else
