[PS2] Fix: References to incorrect cvars

diff --git a/Engine/Plugins/Media/PixelStreaming2/Source/PixelStreaming2HMD/Private/PixelStreaming2HMD.cpp b/Engine/Plugins/Media/PixelStreaming2/Source/PixelStreaming2HMD/Private/PixelStreaming2HMD.cpp
--- a/Engine/Plugins/Media/PixelStreaming2/Source/PixelStreaming2HMD/Private/PixelStreaming2HMD.cpp
+++ b/Engine/Plugins/Media/PixelStreaming2/Source/PixelStreaming2HMD/Private/PixelStreaming2HMD.cpp
@@ -96,34 +96,19 @@ void FPixelStreaming2HMD::SetEyeViews(FTransform Left, FMatrix LeftProj, FTransf
 	if (!bReceivedTransforms)
 	{
 		// Couple engine's render rate and streaming rate
-		if (IConsoleVariable* DecoupleFramerateCVar = IConsoleManager::Get().FindConsoleVariable(TEXT("PixelStreaming.DecoupleFramerate")))
-		{
-			DecoupleFramerateCVar->Set(false);
-		}
+		UPixelStreaming2PluginSettings::CVarDecoupleFramerate->Set(false);
 
 		// Set the rate at which we will stream
-		if (IConsoleVariable* StreamFPS = IConsoleManager::Get().FindConsoleVariable(TEXT("PixelStreaming.WebRTC.Fps")))
-		{
-			StreamFPS->Set(90);
-		}
+		UPixelStreaming2PluginSettings::CVarWebRTCFps->Set(90);
 
-		// Set the MinQP to bound quality
-		if (IConsoleVariable* MinQP = IConsoleManager::Get().FindConsoleVariable(TEXT("PixelStreaming.Encoder.MinQP")))
-		{
-			MinQP->Set(15);
-		}
+		// Set the MaxQuality to bound quality
+		UPixelStreaming2PluginSettings::CVarEncoderMaxQuality->Set(70);
 
 		// Necessary for coupled framerate
-		if (IConsoleVariable* CaptureUseFence = IConsoleManager::Get().FindConsoleVariable(TEXT("PixelStreaming.CaptureUseFence")))
-		{
-			CaptureUseFence->Set(false);
-		}
+		UPixelStreaming2PluginSettings::CVarCaptureUseFence->Set(false);
 
 		// Disable keyframes interval, only send them as needed
-		if (IConsoleVariable* KeyframeInterval = IConsoleManager::Get().FindConsoleVariable(TEXT("PixelStreaming.Encoder.KeyframeInterval")))
-		{
-			KeyframeInterval->Set(0);
-		}
+		UPixelStreaming2PluginSettings::CVarEncoderKeyframeInterval->Set(-1);
 	}
 
 	bReceivedTransforms = true;
