[PS2] Fix: ini not working when packaged

Gist of it came down to that when using the ConsoleVariable metadata, it would save the setting in the ini as the cvar name instead of property name, with the property value. This would work for 99% of cases but any time there was a mismatch between what the cvar stored value and the property, we would get bugs (eg CVar storing string but property storing enum). I've fixed this by removing the problem metadata from all properties and the ini now stores the properties as the property name and property value.

On startup, we:
Load the ini and set the corresponding cvar to the value in the ini accounting for any conversions (eg enum to string, TArray<FString> to comma separated list)
Verify command line args to log any invalid options to users
Parse the command line and update both the property and cvar with the command line arg
Parse the command line for any legacy launch options and update the corresponding property and cvar

This CL also fixes a few bugs:
- Warnings about the DefaultPixelStreaming2Config.ini not being explicitly included or excluded. This warning required every user to explicitly add the config to their DefaultGame.ini file before packaging. By moving the ini values to be stored in DefaultGame this warning is no longer present
- Fixed a warning about the BitmaskEnum on the WebRTCPortAllocatorFlags property. UE wanted the full script path to the UENUM and not the c++ type
- WebRTCCodecPreferences now only lets you add a codec if it isn't already in the list. Previously, you could have duplicate values

diff --git a/Engine/Plugins/Media/PixelStreaming2/Source/PixelStreaming2Settings/Internal/PixelStreaming2PluginSettings.h b/Engine/Plugins/Media/PixelStreaming2/Source/PixelStreaming2Settings/Internal/PixelStreaming2PluginSettings.h
--- a/Engine/Plugins/Media/PixelStreaming2/Source/PixelStreaming2Settings/Internal/PixelStreaming2PluginSettings.h
+++ b/Engine/Plugins/Media/PixelStreaming2/Source/PixelStreaming2Settings/Internal/PixelStreaming2PluginSettings.h
@@ -65,7 +65,7 @@ enum class EInputControllerMode : uint8
 
 // Config loaded/saved to an .ini file.
 // It is also exposed through the plugin settings page in editor.
-UCLASS(config = PixelStreaming2, defaultconfig, meta = (DisplayName = "PixelStreaming2"))
+UCLASS(config = Game, defaultconfig, meta = (DisplayName = "PixelStreaming2"))
 class PIXELSTREAMING2SETTINGS_API UPixelStreaming2PluginSettings : public UDeveloperSettings
 {
 	GENERATED_BODY()
@@ -76,7 +76,6 @@ public:
 	// clang-format off
 	static TAutoConsoleVariable<bool> CVarLogStats;
 	UPROPERTY(config, EditAnywhere, Category = "PixelStreaming", meta = (
-		ConsoleVariable = "PixelStreaming2.LogStats",
 		DisplayName = "Log Pixel Streaming Stats",
 		ToolTip = "Whether to show PixelStreaming stats in the log (default: false)."
 		))
@@ -84,7 +83,6 @@ public:
 
 	static TAutoConsoleVariable<bool> CVarDisableLatencyTester;
 	UPROPERTY(config, EditAnywhere, Category = "PixelStreaming", meta = (
-		ConsoleVariable = "PixelStreaming2.DisableLatencyTester",
 		DisplayName = "Disable Latency Tester",
 		ToolTip = "If true disables latency tester being triggerable."
 		))
@@ -92,7 +90,6 @@ public:
 
 	static TAutoConsoleVariable<FString> CVarInputController;
 	UPROPERTY(config, EditAnywhere, Category = "PixelStreaming", meta = (
-		MappedConsoleVariable = "PixelStreaming2.InputController",
 		DisplayName = "Input Controller Mode",
 		ToolTip = "If true disables latency tester being triggerable."
 		))
@@ -100,7 +97,6 @@ public:
 
 	static TAutoConsoleVariable<bool> CVarDecoupleFramerate;
 	UPROPERTY(config, EditAnywhere, Category = "PixelStreaming", meta = (
-		ConsoleVariable = "PixelStreaming2.DecoupleFrameRate",
 		DisplayName = "Decouple Frame Rate",
 		ToolTip = "Whether we should only stream as fast as we render or at some fixed interval. Coupled means only stream what we render."
 		))
@@ -108,7 +104,6 @@ public:
 
 	static TAutoConsoleVariable<float> CVarDecoupleWaitFactor;
 	UPROPERTY(config, EditAnywhere, Category = "PixelStreaming", meta = (
-		ConsoleVariable = "PixelStreaming2.DecoupleWaitFactor",
 		DisplayName = "Decouple Wait Factor",
 		ToolTip = "Frame rate factor to wait for a captured frame when streaming in decoupled mode. Higher factor waits longer but may also result in higher latency."
 		))
@@ -116,7 +111,6 @@ public:
 
 	static TAutoConsoleVariable<float> CVarSignalingReconnectInterval;
 	UPROPERTY(config, EditAnywhere, Category = "PixelStreaming", meta = (
-		ConsoleVariable = "PixelStreaming2.SignalingReconnectInterval",
 		DisplayName = "Signaling Reconnect Interval",
 		ToolTip = "Frame rate factor to wait for a captured frame when streaming in decoupled mode. Higher factor waits longer but may also result in higher latency."
 		))
@@ -124,7 +118,6 @@ public:
 
 	static TAutoConsoleVariable<float> CVarSignalingKeepAliveInterval;
 	UPROPERTY(config, EditAnywhere, Category = "PixelStreaming", meta = (
-		ConsoleVariable = "PixelStreaming2.SignalingKeepAliveInterval",
 		DisplayName = "Signaling Reconnect Interval",
 		ToolTip = "Frame rate factor to wait for a captured frame when streaming in decoupled mode. Higher factor waits longer but may also result in higher latency."
 		))
@@ -132,7 +125,6 @@ public:
 
 	static TAutoConsoleVariable<bool> CVarUseMediaCapture;
 	UPROPERTY(config, EditAnywhere, Category = "PixelStreaming", meta = (
-		ConsoleVariable = "PixelStreaming2.UseMediaCapture",
 		DisplayName = "Use Media Capture",
 		ToolTip = "Use Media Capture from MediaIOFramework to capture frames rather than Pixel Streamings internal backbuffer sources."
 		))
@@ -140,7 +132,6 @@ public:
 
 	static TAutoConsoleVariable<FString> CVarDefaultStreamerID;
 	UPROPERTY(config, EditAnywhere, Category = "PixelStreaming", meta = (
-		ConsoleVariable = "PixelStreaming2.ID",
 		DisplayName = "Default Streamer ID",
 		ToolTip = "Default Streamer ID to be used when not specified elsewhere."
 		))
@@ -148,7 +139,6 @@ public:
 
 	static TAutoConsoleVariable<FString> CVarConnectionURL;
 	UPROPERTY(config, EditAnywhere, Category = "PixelStreaming", meta = (
-		ConsoleVariable = "PixelStreaming2.ConnectionURL",
 		DisplayName = "Default Connection URL",
 		ToolTip = "Default URL to connect to. This can be a URL to a signalling server or some other endpoint with the format (protocol)://(host):(port)"
 		))
@@ -158,7 +148,6 @@ public:
 
 	static TAutoConsoleVariable<bool> CVarCaptureUseFence;
 	UPROPERTY(config, EditAnywhere, Category = "PixelStreaming", meta = (
-		ConsoleVariable = "PixelStreaming2.CaptureUseFence",
 		DisplayName = "Capture Using Fence",
 		ToolTip = "Whether the texture copy we do during image capture should use a fence or not (non-fenced is faster but less safe)."
 		))
@@ -189,7 +178,6 @@ public:
 	// Begin Encoder Settings
 	static TAutoConsoleVariable<int32> CVarEncoderTargetBitrate;
 	UPROPERTY(config, EditAnywhere, Category = "EncoderSettings", meta = (
-		ConsoleVariable = "PixelStreaming2.Encoder.TargetBitrate",
 		DisplayName = "Target Bitrate",
 		ToolTip = "Target bitrate (bps). Ignore the bitrate WebRTC wants (not recommended). Set to -1 to disable. Default -1.",
 		ClampMin = -1))
@@ -197,7 +185,6 @@ public:
 
 	static TAutoConsoleVariable<int32>	 CVarEncoderMinQuality;
 	UPROPERTY(config, EditAnywhere, Category = "EncoderSettings", meta = (
-		ConsoleVariable = "PixelStreaming2.Encoder.MinQuality",
 		DisplayName = "Encoder Minimum Quality",
 		ToolTip = "0-100, Higher values result in a better minimum quality but higher average bitrates. Default 0 - i.e. no limit on a minimum Quality.",
 		ClampMin = 0, ClampMax = 100))
@@ -205,7 +192,6 @@ public:
 
 	static TAutoConsoleVariable<int32>	CVarEncoderMaxQuality;
 	UPROPERTY(config, EditAnywhere, Category = "EncoderSettings", meta = (
-		ConsoleVariable = "PixelStreaming2.Encoder.MaxQuality",
 		DisplayName = "Encoder Maximum Quality",
 		ToolTip = "0-100, Lower values result in lower average bitrates but reduces maximum quality. Default 100 - i.e. no limit on a maximum QP.",
 		ClampMin = 0, ClampMax = 100))
@@ -214,7 +200,6 @@ public:
 	// not directly connected to cvar due to string to enum conversion
 	static TAutoConsoleVariable<FString> CVarEncoderQualityPreset;
 	UPROPERTY(config, EditAnywhere, Category = "EncoderSettings", meta = (
-		MappedConsoleVariable = "PixelStreaming2.Encoder.QualityPreset",
 		DisplayName = "Default Encoding Quality Profile",
 		ToolTip = "PixelStreaming Encoder presets that affecting Quality vs Bitrate. Supported modes are `ULTRA_LOW_QUALITY`, `LOW_QUALITY`, `DEFAULT`, `HIGH_QUALITY`, `LOSSLESS`"
 		))
@@ -223,7 +208,6 @@ public:
 	// not directly connected to cvar due to string to enum conversion
 	static TAutoConsoleVariable<FString> CVarEncoderLatencyMode;
 	UPROPERTY(config, EditAnywhere, Category = "EncoderSettings", meta = (
-		MappedConsoleVariable = "PixelStreaming2.Encoder.LatencyMode",
 		DisplayName = "Default Encoding Quality Profile",
 		ToolTip = "PixelStreaming Encoder presets that affecting Quality vs Latency. Supported modes are `ULTRA_LOW_LATENCY`, `LOW_LATENCY`, `DEFAULT`"
 		))
@@ -231,7 +215,6 @@ public:
 
 	static TAutoConsoleVariable<int32> CVarEncoderKeyframeInterval;
 	UPROPERTY(config, EditAnywhere, Category = "EncoderSettings", meta = (
-		ConsoleVariable = "PixelStreaming2.Encoder.KeyframeInterval",
 		DisplayName = "Default Keyframe Interval",
 		ToolTip = "How many frames before a key frame is sent. Default: -1 which disables the sending of periodic key frames. Note: NVENC reqires a reinitialization when this changes."
 		))
@@ -239,7 +222,6 @@ public:
 
 	static TAutoConsoleVariable<int32> CVarEncoderMaxSessions;
 	UPROPERTY(config, EditAnywhere, Category = "EncoderSettings", meta = (
-		ConsoleVariable = "PixelStreaming2.Encoder.MaxSessions",
 		DisplayName = "Default Max Number of Encoding Session",
 		ToolTip = "-1 implies no limit. Maximum number of concurrent hardware encoder sessions for Pixel Streaming. Note: Geforce gpus only support 8 concurrent sessions and will rollover to software encoding when that number is exceeded."
 		))
@@ -247,7 +229,6 @@ public:
 
 	static TAutoConsoleVariable<bool> CVarEncoderEnableSimulcast;
 	UPROPERTY(config, EditAnywhere, Category = "EncoderSettings", meta = (
-		ConsoleVariable = "PixelStreaming2.Encoder.EnableSimulcast",
 		DisplayName = "Enable Simulcast",
 		ToolTip = "Enables simulcast. When enabled, the encoder will encode at full resolution, 1/2 resolution and 1/4 resolution simultaneously. Note: Simulcast is only supported with `H264` and `VP8` and you must use the SFU from the infrastructure to fully utilise this functionality.",
 		EditCondition="!WebRTCNegotiateCodecs",
@@ -257,7 +238,6 @@ public:
 
 	static TAutoConsoleVariable<FString> CVarEncoderCodec;
 	UPROPERTY(config, EditAnywhere, Category = "EncoderSettings", meta = (
-		ConsoleVariable = "PixelStreaming2.Encoder.Codec",
 		DisplayName = "Preferred Encoder Codec",
 		ToolTip = "Preferred encoder codec signalled during connection establishment.",
 		GetOptions = "GetVideoCodecOptions",
@@ -268,7 +248,6 @@ public:
 
 	static TAutoConsoleVariable<FString> CVarEncoderScalabilityMode;
 	UPROPERTY(config, EditAnywhere, Category = "EncoderSettings", meta = (
-		ConsoleVariable = "PixelStreaming2.Encoder.ScalabilityMode",
 		DisplayName = "Default Scalability Mode",
 		ToolTip = "Indicates number of spatial and temporal layers used, default: L1T1. For a full list of values refer to https://www.w3.org/TR/webrtc-svc/#scalabilitymodes*",
 		GetOptions = "GetScalabilityModeOptions",
@@ -279,7 +258,6 @@ public:
 
 	static TAutoConsoleVariable<FString> CVarEncoderH264Profile;
 	UPROPERTY(config, EditAnywhere, Category = "EncoderSettings|H264", meta = (
-		MappedConsoleVariable = "PixelStreaming2.Encoder.H264Profile",
 		DisplayName = "Default H264 profile",
 		ToolTip = "PixelStreaming Encoder profile. Supported modes are `AUTO`, `BASELINE`, `MAIN`, `HIGH`, `HIGH444`, `PROGRESSIVE_HIGH` or `CONSTRAINED_HIGH`",
 		ValidEnumValues = "Auto, Baseline, Main, High, High444, ProgressiveHigh, ConstrainedHigh"
@@ -292,7 +270,6 @@ public:
 	// Begin WebRTC CVars
 	static TAutoConsoleVariable<int32> CVarWebRTCFps;
 	UPROPERTY(config, EditAnywhere, Category = "WebRtcSettings", meta = (
-		ConsoleVariable = "PixelStreaming2.WebRTC.Fps",
 		DisplayName = "Default WebRtc FPS",
 		ToolTip = "Framerate for WebRTC encoding. Default: 60"
 		))
@@ -300,7 +277,6 @@ public:
 
 	static TAutoConsoleVariable<int32> CVarWebRTCStartBitrate;
 	UPROPERTY(config, EditAnywhere, Category = "WebRtcSettings", meta = (
-		ConsoleVariable = "PixelStreaming2.WebRTC.StartBitrate",
 		DisplayName = "Default WebRtc Start Bitrate",
 		ToolTip = "FStart bitrate (bps) that WebRTC will try begin the stream with. Must be between Min/Max bitrates. Default: 1000000"
 		))
@@ -308,7 +284,6 @@ public:
 
 	static TAutoConsoleVariable<int32> CVarWebRTCMinBitrate;
 	UPROPERTY(config, EditAnywhere, Category = "WebRtcSettings", meta = (
-		ConsoleVariable = "PixelStreaming2.WebRTC.MinBitrate",
 		DisplayName = "Default WebRtc Minimum Bitrate",
 		ToolTip = "Min bitrate (bps) that WebRTC will not request below. Careful not to set too high otherwise WebRTC will just drop frames. Default: 100000"
 		))
@@ -316,7 +291,6 @@ public:
 
 	static TAutoConsoleVariable<int32> CVarWebRTCMaxBitrate;
 	UPROPERTY(config, EditAnywhere, Category = "WebRtcSettings", meta = (
-		ConsoleVariable = "PixelStreaming2.WebRTC.MaxBitrate",
 		DisplayName = "Default WebRtc Maximum Bitrate",
 		ToolTip = "Max bitrate (bps) that WebRTC will not request above. Default: 40000000 aka 40 megabits/per second."
 		))
@@ -324,7 +298,6 @@ public:
 
 	static TAutoConsoleVariable<bool> CVarWebRTCDisableReceiveAudio;
 	UPROPERTY(config, EditAnywhere, Category = "WebRtcSettings", meta = (
-		ConsoleVariable = "PixelStreaming2.WebRTC.DisableReceiveAudio",
 		DisplayName = "Disable receiving audio by default",
 		ToolTip = "Disables receiving audio from the browser into UE."
 		))
@@ -332,7 +305,6 @@ public:
 
 	static TAutoConsoleVariable<bool> CVarWebRTCDisableReceiveVideo;
 	UPROPERTY(config, EditAnywhere, Category = "WebRtcSettings", meta = (
-		ConsoleVariable = "PixelStreaming2.WebRTC.DisableReceiveVideo",
 		DisplayName = "Disable receiving video by default",
 		ToolTip = "Disables receiving video from the browser into UE."
 		))
@@ -340,7 +312,6 @@ public:
 
 	static TAutoConsoleVariable<bool> CVarWebRTCDisableTransmitAudio;
 	UPROPERTY(config, EditAnywhere, Category = "WebRtcSettings", meta = (
-		ConsoleVariable = "PixelStreaming2.WebRTC.DisableTransmitAudio",
 		DisplayName = "Disable transmiting audio by default",
 		ToolTip = "Disables transmiting audio to the browser."
 		))
@@ -348,7 +319,6 @@ public:
 
 	static TAutoConsoleVariable<bool> CVarWebRTCDisableTransmitVideo;
 	UPROPERTY(config, EditAnywhere, Category = "WebRtcSettings", meta = (
-		ConsoleVariable = "PixelStreaming2.WebRTC.DisableTransmitVideo",
 		DisplayName = "Disable transmiting video by default",
 		ToolTip = "Disables transmiting video to the browser."
 		))
@@ -356,7 +326,6 @@ public:
 
 	static TAutoConsoleVariable<bool> CVarWebRTCDisableAudioSync;
 	UPROPERTY(config, EditAnywhere, Category = "WebRtcSettings", meta = (
-		ConsoleVariable = "PixelStreaming2.WebRTC.DisableAudioSync",
 		DisplayName = "Disable Audio/Video track sychronisation",
 		ToolTip = "Disables the synchronization of audio and video tracks in WebRTC. This can be useful in low latency usecases where synchronization is not required."
 		))
@@ -364,7 +333,6 @@ public:
 
 	static TAutoConsoleVariable<bool> CVarWebRTCEnableFlexFec;
 	UPROPERTY(config, EditAnywhere, Category = "WebRtcSettings", meta = (
-		ConsoleVariable = "PixelStreaming2.WebRTC.EnableFlexFec",
 		DisplayName = "Enable Flexible Forward Error Correction",
 		ToolTip = "Signals support for Flexible Forward Error Correction to WebRTC. This can cause a reduction in quality if total bitrate is low."
 		))
@@ -372,7 +340,6 @@ public:
 
 	static TAutoConsoleVariable<bool> CVarWebRTCDisableStats;
 	UPROPERTY(config, EditAnywhere, Category = "WebRtcSettings", meta = (
-		ConsoleVariable = "PixelStreaming2.WebRTC.DisableStats",
 		DisplayName = "Disable WebRtc Stats",
 		ToolTip = "Disables the collection of WebRTC stats."
 		))
@@ -380,7 +347,6 @@ public:
 
 	static TAutoConsoleVariable<float> CVarWebRTCStatsInterval;
 	UPROPERTY(config, EditAnywhere, Category = "WebRtcSettings", meta = (
-		ConsoleVariable = "PixelStreaming2.WebRTC.StatsInterval",
 		DisplayName = "Stats Interval",
 		ToolTip = "Configures how often WebRTC stats are collected."
 		))
@@ -388,7 +354,6 @@ public:
 
 	static TAutoConsoleVariable<bool> CVarWebRTCNegotiateCodecs;
 	UPROPERTY(config, EditAnywhere, Category = "WebRtcSettings", meta = (
-		ConsoleVariable = "PixelStreaming2.WebRTC.NegotiateCodecs",
 		DisplayName = "WebRtc Negotiate Codecs",
 		ToolTip = "Whether PixelStreaming should send all its codecs during sdp handshake so peers can negotiate or just send a single selected codec."
 		))
@@ -396,20 +361,24 @@ public:
 
 	static TAutoConsoleVariable<FString> CVarWebRTCCodecPreferences;
 	UPROPERTY(config, EditAnywhere, Category = "WebRtcSettings", meta = (
-		ConsoleVariable = "PixelStreaming2.WebRTC.CodecPreferences",
 		DisplayName = "WebRtc Codec Preferences",
 		ToolTip = "The preference order PixelStreaming will specify during sdp handshake",
 		EditFixedSize,
 		EditCondition="WebRTCNegotiateCodecs",
-		EditConditionHides
+		EditConditionHides,
+		GetOptions = "GetWebRTCCodecPreferencesOptions"
 		))
-	TArray<EVideoCodec> WebRTCCodecPreferences = { EVideoCodec::AV1, EVideoCodec::H264, EVideoCodec::VP9, EVideoCodec::VP8 };
+	TArray<FString> WebRTCCodecPreferences = { 
+		UE::PixelStreaming2::GetCVarStringFromEnum(EVideoCodec::AV1), 
+		UE::PixelStreaming2::GetCVarStringFromEnum(EVideoCodec::H264), 
+		UE::PixelStreaming2::GetCVarStringFromEnum(EVideoCodec::VP9), 
+		UE::PixelStreaming2::GetCVarStringFromEnum(EVideoCodec::VP8)
+	 };
 
 	static TArray<EVideoCodec> GetCodecPreferences();
 
 	static TAutoConsoleVariable<float> CVarWebRTCAudioGain;
 	UPROPERTY(config, EditAnywhere, Category = "WebRtcSettings", meta = (
-		ConsoleVariable = "PixelStreaming2.WebRTC.AudioGain",
 		DisplayName = "Default Audio Gain",
 		ToolTip = "Sets the amount of gain to apply to audio. Default: 1.0"
 		))
@@ -418,8 +387,7 @@ public:
 	static TAutoConsoleVariable<FString> CVarWebRTCPortAllocatorFlags;
 	UPROPERTY(config, EditAnywhere, Category = "WebRtcSettings", meta = (
 		Bitmask,
-		BitmaskEnum = "EPortAllocatorFlags",
-		ConsoleVariable = "PixelStreaming2.WebRTC.PortAllocatorFlags",
+		BitmaskEnum = "/Script/PixelStreaming2Settings.EPortAllocatorFlags",
 		DisplayName = "WebRtc port allocation flags",
 		ToolTip = "Sets the WebRTC port allocator flags. See "
 		))
@@ -429,7 +397,6 @@ public:
 
 	static TAutoConsoleVariable<int> CVarWebRTCMinPort;
 	UPROPERTY(config, EditAnywhere, Category = "WebRtcSettings", meta = (
-		ConsoleVariable = "PixelStreaming2.WebRTC.MinPort",
 		DisplayName = "Default WebRtc Min Port",
 		ToolTip = "Sets the minimum usable port for the WebRTC port allocator. Default: 49152"
 		))
@@ -437,7 +404,6 @@ public:
 
 	static TAutoConsoleVariable<int> CVarWebRTCMaxPort;
 	UPROPERTY(config, EditAnywhere, Category = "WebRtcSettings", meta = (
-		ConsoleVariable = "PixelStreaming2.WebRTC.MaxPort",
 		DisplayName = "Default WebRtc Max Port",
 		ToolTip = "Sets the maximum usable port for the WebRTC port allocator. Default: 65535"
 		))
@@ -445,7 +411,6 @@ public:
 
 	static TAutoConsoleVariable<FString> CVarWebRTCFieldTrials;
 	UPROPERTY(config, EditAnywhere, Category = "WebRtcSettings|Field Trials", meta = (
-		ConsoleVariable = "PixelStreaming2.WebRTC.FieldTrials",
 		DisplayName = "Manually defined WebRtc Field Trials",
 		ToolTip = "Sets the WebRTC field trials string. Format:\"TRIAL1/VALUE1/TRIAL2/VALUE2/\" see https://webrtc.googlesource.com/src/+/HEAD/g3doc/field-trials.md"
 		))
@@ -453,7 +418,6 @@ public:
 
 	static TAutoConsoleVariable<bool> CVarWebRTCDisableFrameDropper;
 	UPROPERTY(config, EditAnywhere, Category = "WebRtcSettings|Field Trials", meta = (
-		ConsoleVariable = "PixelStreaming2.WebRTC.DisableFrameDropper",
 		DisplayName = "Disable Frame Dropper",
 		ToolTip = "Disables the WebRTC internal frame dropper using the field trial WebRTC-FrameDropper/Disabled/"
 		))
@@ -461,7 +425,6 @@ public:
 
 	static TAutoConsoleVariable<float> CVarWebRTCVideoPacingMaxDelay;
 	UPROPERTY(config, EditAnywhere, Category = "WebRtcSettings|Field Trials", meta = (
-		ConsoleVariable = "PixelStreaming2.WebRTC.VideoPacing.MaxDelay",
 		DisplayName = "Video Pacing Max Delay",
 		ToolTip = "Enables the WebRTC-Video-Pacing field trial and sets the max delay (ms) parameter. Default: -1.0f (values below zero are discarded.)"
 		))
@@ -469,7 +432,6 @@ public:
 
 	static TAutoConsoleVariable<float> CVarWebRTCVideoPacingFactor;
 	UPROPERTY(config, EditAnywhere, Category = "WebRtcSettings|Field Trials", meta = (
-		ConsoleVariable = "PixelStreaming2.WebRTC.VideoPacing.Factor",
 		DisplayName = "Video Pacing Factor",
 		ToolTip = "Enables the WebRTC-Video-Pacing field trial and sets the video pacing factor parameter. Larger values are more lenient on larger bitrates. Default: -1.0f (values below zero are discarded.)"
 		))
@@ -480,7 +442,6 @@ public:
 	// Begin EditorStreaming CVars
 	static TAutoConsoleVariable<bool> CVarEditorStartOnLaunch;
 	UPROPERTY(config, EditAnywhere, Category = "EditorStreaming", meta = (
-		ConsoleVariable = "PixelStreaming2.Editor.StartOnLaunch",
 		DisplayName = "Start On Launch",
 		ToolTip = "Start Editor Streaming as soon as the Unreal Editor is launched. Default: false"
 		))
@@ -488,7 +449,6 @@ public:
 
 	static TAutoConsoleVariable<bool> CVarEditorUseRemoteSignallingServer;
 	UPROPERTY(config, EditAnywhere, Category = "EditorStreaming", meta = (
-		ConsoleVariable = "PixelStreaming2.Editor.UseRemoteSignallingServer",
 		DisplayName = "Use Remote Signalling Server",
 		ToolTip = "Enables the use of a remote signalling server. Default: false"
 		))
@@ -496,7 +456,6 @@ public:
 
 	static TAutoConsoleVariable<FString> CVarEditorSource;
 	UPROPERTY(config, EditAnywhere, Category = "EditorStreaming", meta = (
-		MappedConsoleVariable = "PixelStreaming2.Editor.Source",
 		DisplayName = "Source",
 		ToolTip = "Editor PixelStreaming source. Supported values are `Editor`, `LevelEditorViewport`. Default: Editor"
 		))
@@ -506,7 +465,6 @@ public:
 	// Begin HMD CVars
 	static TAutoConsoleVariable<bool> CVarHMDEnable;
 	UPROPERTY(config, EditAnywhere, Category = "XR Streaming", meta = (
-		ConsoleVariable = "PixelStreaming2.HMD.Enable",
 		DisplayName = "Enable HMD",
 		ToolTip = "Enables HMD specific functionality for Pixel Streaming. Namely input handling and stereoscopic rendering. Default: false"
 		))
@@ -514,7 +472,6 @@ public:
 
 	static TAutoConsoleVariable<bool> CVarHMDMatchAspectRatio;
 	UPROPERTY(config, EditAnywhere, Category = "XR Streaming", meta = (
-		ConsoleVariable = "PixelStreaming2.HMD.MatchAspectRatio",
 		DisplayName = "Match Aspect Ratio",
 		ToolTip = "If true automatically resize the rendering resolution to match the aspect ratio determined by the HFoV and VFoV. Default: true"
 		))
@@ -522,7 +479,6 @@ public:
 
 	static TAutoConsoleVariable<bool> CVarHMDApplyEyePosition;
 	UPROPERTY(config, EditAnywhere, Category = "XR Streaming", meta = (
-		ConsoleVariable = "PixelStreaming2.HMD.ApplyEyePosition",
 		DisplayName = "Apply Eye Position",
 		ToolTip = "If true automatically position each eye's rendering by whatever amount WebXR reports for each left-right XRView. If false do no eye positioning. Default: true"
 		))
@@ -530,7 +486,6 @@ public:
 
 	static TAutoConsoleVariable<bool> CVarHMDApplyEyeRotation;
 	UPROPERTY(config, EditAnywhere, Category = "XR Streaming", meta = (
-		ConsoleVariable = "PixelStreaming2.HMD.ApplyEyeRotation",
 		DisplayName = "Apply Eye Position",
 		ToolTip = "If true automatically rotate each eye's rendering by whatever amount WebXR reports for each left-right XRView. If false do no eye rotation. Default: true"
 		))
@@ -538,7 +493,6 @@ public:
 
 	static TAutoConsoleVariable<float> CVarHMDHFOV;
 	UPROPERTY(config, EditAnywhere, Category = "XR Streaming", meta = (
-		ConsoleVariable = "PixelStreaming2.HMD.HFOV",
 		DisplayName = "Horizontal FOV Override",
 		ToolTip = "Overrides the horizontal field of view for HMD rendering, values are in degrees and values less than 0.0f disable the override. Default: -1.0f"
 		))
@@ -546,7 +500,6 @@ public:
 
 	static TAutoConsoleVariable<float> CVarHMDVFOV;
 	UPROPERTY(config, EditAnywhere, Category = "XR Streaming", meta = (
-		ConsoleVariable = "PixelStreaming2.HMD.VFOV",
 		DisplayName = "Vertical FOV Override",
 		ToolTip = "Overrides the vertical field of view for HMD rendering, values are in degrees and values less than 0.0f disable the override. Default: -1.0f"
 		))
@@ -554,7 +507,6 @@ public:
 
 	static TAutoConsoleVariable<float> CVarHMDIPD;
 	UPROPERTY(config, EditAnywhere, Category = "XR Streaming", meta = (
-		ConsoleVariable = "PixelStreaming2.HMD.IPD",
 		DisplayName = "Interpupillary Distance Override",
 		ToolTip = "Overrides the HMD IPD (interpupillary distance), values are in centimeters and values less than 0.0f disable the override. Default: -1.0f"
 		))
@@ -562,7 +514,6 @@ public:
 
 	static TAutoConsoleVariable<float> CVarHMDProjectionOffsetX;
 	UPROPERTY(config, EditAnywhere, Category = "XR Streaming", meta = (
-		ConsoleVariable = "PixelStreaming2.HMD.ProjectionOffsetX",
 		DisplayName = "Horizontal Projection Offset Override",
 		ToolTip = "Overrides the left/right eye projection matrix x-offset, values are in clip space and values less than 0.0f disable the override. Default: -1.0f"
 		))
@@ -570,7 +521,6 @@ public:
 
 	static TAutoConsoleVariable<float> CVarHMDProjectionOffsetY;
 	UPROPERTY(config, EditAnywhere, Category = "XR Streaming", meta = (
-		ConsoleVariable = "PixelStreaming2.HMD.ProjectionOffsetY",
 		DisplayName = "Vertical Projection Offset Override",
 		ToolTip = "Overrides the left-right eye projection matrix y-offset, values are in clip space and values less than 0.0f disable the override. Default: -1.0f"
 		))
@@ -580,7 +530,6 @@ public:
 	// Begin Input CVars
 	static TAutoConsoleVariable<bool> CVarInputAllowConsoleCommands;
 	UPROPERTY(config, EditAnywhere, Category = "Input", meta = (
-		ConsoleVariable = "PixelStreaming2.AllowPixelStreaming2Commands",
 		DisplayName = "Allow Commands",
 		ToolTip = "If true browser can send consoleCommand payloads that execute in UE's console. Default: false"
 		))
@@ -588,7 +537,6 @@ public:
 
 	static TAutoConsoleVariable<FString> CVarInputKeyFilter;
 	UPROPERTY(config, EditAnywhere, Category = "Input", meta = (
-		ConsoleVariable = "PixelStreaming2.KeyFilter",
 		DisplayName = "Key Filter",
 		ToolTip = "Comma separated list of keys to ignore from streaming clients. Default: \"\""
 		))
@@ -653,9 +601,12 @@ public:
 	static FDelegates* Delegates();
 
 private:
-	void SetCVarFromPropertyAndValue(IConsoleVariable* CVar, FProperty* Property, const FString& CVarString, const FString& Value);
-	void SetCVarFromProperty(IConsoleVariable* CVar, FProperty* Property, const FString& CVarString);
+	void SetCVarAndPropertyFromValue(const FString& CVarName, FProperty* Property, const FString& Value);
+	void SetCVarFromProperty(const FString& CVarName, FProperty* Property);
+
+	void InitializeCVarsFromProperties();
 	void ValidateCommandLineArgs();
+	void ParseCommandlineArgs();
 	void ParseLegacyCommandlineArgs();
 
 	static FDelegates* DelegateSingleton;
@@ -665,4 +616,7 @@ private:
 
 	UFUNCTION()
 	TArray<FString> GetScalabilityModeOptions() const;
+
+	UFUNCTION()
+	TArray<FString> GetWebRTCCodecPreferencesOptions() const;
 };
\ No newline at end of file
diff --git a/Engine/Plugins/Media/PixelStreaming2/Source/PixelStreaming2Settings/Private/PixelStreaming2PluginSettings.cpp b/Engine/Plugins/Media/PixelStreaming2/Source/PixelStreaming2Settings/Private/PixelStreaming2PluginSettings.cpp
--- a/Engine/Plugins/Media/PixelStreaming2/Source/PixelStreaming2Settings/Private/PixelStreaming2PluginSettings.cpp
+++ b/Engine/Plugins/Media/PixelStreaming2/Source/PixelStreaming2Settings/Private/PixelStreaming2PluginSettings.cpp
@@ -58,39 +58,6 @@ namespace
 		}
 	}
 
-	static void VerifyCodecPreferenceSettings(UPixelStreaming2PluginSettings* This, FProperty* Property, IConsoleVariable* CVar, const FString& Value = TEXT(""))
-	{
-		TArray<FString> ValidStringArray;
-
-		if (Value != "")
-		{
-			TArray<FString> StringArray;
-			Value.ParseIntoArray(StringArray, TEXT(","));
-
-			for (const FString& CodecString : StringArray)
-			{
-				if (StaticEnum<EVideoCodec>()->GetIndexByNameString(CodecString) == INDEX_NONE)
-				{
-					UE_LOGFMT(LogPixelStreaming2Settings, Warning, "Invalid value {0} received for enum of type {1}", CodecString, StaticEnum<EVideoCodec>()->GetName());
-					continue;
-				}
-
-				ValidStringArray.Add(CodecString);
-			}
-		}
-		else
-		{
-			FArrayProperty*		ArrayProperty = CastField<FArrayProperty>(Property);
-			TArray<EVideoCodec> CodecArray = *ArrayProperty->ContainerPtrToValuePtr<TArray<EVideoCodec>>(This);
-			for (EVideoCodec Codec : CodecArray)
-			{
-				ValidStringArray.Add(UE::PixelStreaming2::GetCVarStringFromEnum(Codec));
-			}
-		}
-
-		CVar->Set(*FString::Join(ValidStringArray, TEXT(",")), ECVF_SetByConsole);
-	}
-
 	FString ConsoleVariableToCommandArgValue(const FString InCVarName)
 	{
 		// CVars are . deliminated by section. To get their equivilent commandline arg for parsing
@@ -134,10 +101,115 @@ namespace
 			CVar->Set(true, ECVF_SetByCommandline);
 		}
 	}
+
+	static FString FindPropertyFromCVar(const TSet<TPair<FString, FString>> Set, const FString& Key)
+	{
+		for (const TPair<FString, FString>& Pair : Set)
+		{
+			if (Pair.Key == Key)
+			{
+				return Pair.Value;
+			}
+		}
+
+		return "";
+	}
+
+	static FString FindCVarFromProperty(const TSet<TPair<FString, FString>> Set, const FString& Value)
+	{
+		for (const TPair<FString, FString>& Pair : Set)
+		{
+			if (Pair.Value == Value)
+			{
+				return Pair.Key;
+			}
+		}
+
+		return "";
+	}
 } // namespace
 
-static FName PixelStreaming2ConsoleVariableMetaFName(TEXT("ConsoleVariable"));
-static FName PixelStreaming2MappedConsoleVariableFName(TEXT("MappedConsoleVariable"));
+// Map of Property Names to their commandline args as GetMetaData() is not avaliable in packaged projects
+static const TSet<TPair<FString, FString>> GetCmdArg = {
+	{ "PixelStreaming2.LogStats", "LogStats" },
+	{ "PixelStreaming2.SendPlayerIdAsInteger", "SendPlayerIdAsInteger" },
+	{ "PixelStreaming2.DisableLatencyTester", "DisableLatencyTester" },
+	{ "PixelStreaming2.DecoupleFrameRate", "DecoupleFramerate" },
+	{ "PixelStreaming2.DecoupleWaitFactor", "DecoupleWaitFactor" },
+	{ "PixelStreaming2.SignalingReconnectInterval", "SignalingReconnectInterval" },
+	{ "PixelStreaming2.SignalingKeepAliveInterval", "SignalingKeepAliveInterval" },
+	{ "PixelStreaming2.UseMediaCapture", "UseMediaCapture" },
+	{ "PixelStreaming2.ID", "DefaultStreamerID" },
+	{ "PixelStreaming2.ConnectionURL", "ConnectionURL" },
+	{ "PixelStreaming2.CaptureUseFence", "CaptureUseFence" },
+	{ "PixelStreaming2.Encoder.Codec", "Codec" },
+	{ "PixelStreaming2.Encoder.TargetBitrate", "EncoderTargetBitrate" },
+	{ "PixelStreaming2.Encoder.MinQuality", "EncoderMinQuality" },
+	{ "PixelStreaming2.Encoder.MaxQuality", "EncoderMaxQuality" },
+	{ "PixelStreaming2.Encoder.ScalabilityMode", "ScalabilityMode" },
+	{ "PixelStreaming2.Encoder.KeyframeInterval", "KeyframeInterval" },
+	{ "PixelStreaming2.Encoder.MaxSessions", "MaxSessions" },
+	{ "PixelStreaming2.Encoder.EnableSimulcast", "EnableSimulcast" },
+	{ "PixelStreaming2.WebRTC.Fps", "WebRTCFps" },
+	{ "PixelStreaming2.WebRTC.StartBitrate", "WebRTCStartBitrate" },
+	{ "PixelStreaming2.WebRTC.MinBitrate", "WebRTCMinBitrate" },
+	{ "PixelStreaming2.WebRTC.MaxBitrate", "WebRTCMaxBitrate" },
+	{ "PixelStreaming2.WebRTC.DisableReceiveAudio", "WebRTCDisableReceiveAudio" },
+	{ "PixelStreaming2.WebRTC.DisableReceiveVideo", "WebRTCDisableReceiveVideo" },
+	{ "PixelStreaming2.WebRTC.DisableTransmitAudio", "WebRTCDisableTransmitAudio" },
+	{ "PixelStreaming2.WebRTC.DisableTransmitVideo", "WebRTCDisableTransmitVideo" },
+	{ "PixelStreaming2.WebRTC.DisableAudioSync", "WebRTCDisableAudioSync" },
+	{ "PixelStreaming2.WebRTC.EnableFlexFec", "WebRTCEnableFlexFec" },
+	{ "PixelStreaming2.WebRTC.DisableStats", "WebRTCDisableStats" },
+	{ "PixelStreaming2.WebRTC.StatsInterval", "WebRTCStatsInterval" },
+	{ "PixelStreaming2.WebRTC.NegotiateCodecs", "WebRTCNegotiateCodecs" },
+	{ "PixelStreaming2.WebRTC.AudioGain", "WebRTCAudioGain" },
+	{ "PixelStreaming2.WebRTC.PortAllocatorFlags", "WebRTCPortAllocatorFlags" },
+	{ "PixelStreaming2.WebRTC.MinPort", "WebRTCMinPort" },
+	{ "PixelStreaming2.WebRTC.MaxPort", "WebRTCMaxPort" },
+	{ "PixelStreaming2.WebRTC.FieldTrials", "WebRTCFieldTrials" },
+	{ "PixelStreaming2.WebRTC.DisableFrameDropper", "WebRTCDisableFrameDropper" },
+	{ "PixelStreaming2.WebRTC.VideoPacing.MaxDelay", "WebRTCVideoPacingMaxDelay" },
+	{ "PixelStreaming2.WebRTC.VideoPacing.Factor", "WebRTCVideoPacingFactor" },
+	{ "PixelStreaming2.Editor.StartOnLaunch", "EditorStartOnLaunch" },
+	{ "PixelStreaming2.Editor.UseRemoteSignallingServer", "EditorUseRemoteSignallingServer" },
+	{ "PixelStreaming2.HMD.Enable", "HMDEnable" },
+	{ "PixelStreaming2.HMD.MatchAspectRatio", "HMDMatchAspectRatio" },
+	{ "PixelStreaming2.HMD.ApplyEyePosition", "HMDAppleEyePosition" },
+	{ "PixelStreaming2.HMD.ApplyEyeRotation", "HMDApplyEyeRotation" },
+	{ "PixelStreaming2.HMD.HFOV", "HMDHFOV" },
+	{ "PixelStreaming2.HMD.VFOV", "HMDVFOV" },
+	{ "PixelStreaming2.HMD.IPD", "HMDIPD" },
+	{ "PixelStreaming2.HMD.ProjectionOffsetX", "HMDProjectionOffsetX" },
+	{ "PixelStreaming2.HMD.ProjectionOffsetY", "HMDProjectionOffsetY" },
+	{ "PixelStreaming2.AllowPixelStreamingCommands", "InputAllowConsoleCommands" },
+	{ "PixelStreaming2.KeyFilter", "InputKeyFilter" },
+	{ "PixelStreaming2.WebRTC.CodecPreferences", "WebRTCCodecPreferences" }
+};
+
+static const TSet<TPair<FString, FString>> GetMappedCmdArg = {
+	{ "PixelStreaming2.InputController", "InputController" },
+	{ "PixelStreaming2.Encoder.QualityPreset", "QualityPreset" },
+	{ "PixelStreaming2.Encoder.LatencyMode", "LatencyMode" },
+	{ "PixelStreaming2.Encoder.H264Profile", "H264Profile" },
+	{ "PixelStreaming2.Editor.Source", "EditorSource" }
+};
+
+// Map a legacy cvar to its new property
+static const TSet<TPair<FString, FString>> GetLegacyCmdArg = {
+	{ "PixelStreaming2.Encoder.MinQp", "EncoderMaxQuality" },								   // Renamed to MaxQuality
+	{ "PixelStreaming2.Encoder.MaxQp", "EncoderMinQuality" },								   // Renamed to MinQuality
+	{ "PixelStreaming2.IP", "ConnectionURL" },												   // Moved to ConnectionURL
+	{ "PixelStreaming2.Port", "ConnectionURL" },											   // Moved to ConnectionURL
+	{ "PixelStreaming2.URL", "ConnectionURL" },												   // Renamed to ConnectionURL
+	{ "PixelStreaming2.SignallingURL", "ConnectionURL" },									   // Renamed to ConnectionURL
+	{ "AllowPixelStreamingCommands", "InputAllowConsoleCommands" },							   // Renamed to InputAllowConsoleCommands
+	{ "PixelStreaming2.NegotiateCodecs", "WebRTCNegotiateCodecs" },							   // Renamed to PixelStreaming2.WebRTC.NegotiateCodecs
+	{ "PixelStreaming2.EnableHMD", "HMDEnable" },											   // Renamed to PixelStreaming2.HMDEnable
+	{ "Editor.PixelStreaming2.StartOnLaunch", "EditorStartOnLaunch" },						   // Renamed to PixelStreaming2.Editor.StartOnLaunch
+	{ "Editor.PixelStreaming2.UseRemoteSignallingServer", "EditorUseRemoteSignallingServer" }, // Renamed to PixelStreaming2.Editor.UseRemoteSignallingServer
+	{ "Editor.PixelStreaming2.Source", "EditorSource" }										   // Renamed to PixelStreaming2.Editor.Source
+};
 
 // Begin Pixel Streaming Plugin CVars
 TAutoConsoleVariable<bool> UPixelStreaming2PluginSettings::CVarLogStats(
@@ -673,6 +745,15 @@ void SetPortAllocationCVarFromProperty(UObject* This, FProperty* Property)
 	UPixelStreaming2PluginSettings::CVarWebRTCPortAllocatorFlags.AsVariable()->Set(*CVarString, ECVF_SetByProjectSetting);
 }
 
+void SetPortAllocationCVarAndPropertyFromValue(UObject* This, FProperty* Property, const FString& Value)
+{
+	UPixelStreaming2PluginSettings::CVarWebRTCPortAllocatorFlags.AsVariable()->Set(*Value, ECVF_SetByCommandline);
+
+	const FNumericProperty* EnumProperty = CastField<const FNumericProperty>(Property);
+	int64*					PropertyAddress = EnumProperty->ContainerPtrToValuePtr<int64>(This);
+	*PropertyAddress = static_cast<int64>(UPixelStreaming2PluginSettings::GetPortAllocationFlags());
+}
+
 TAutoConsoleVariable<FString> UPixelStreaming2PluginSettings::CVarWebRTCPortAllocatorFlags(
 	TEXT("PixelStreaming2.WebRTC.PortAllocatorFlags"),
 	TEXT(""),
@@ -737,44 +818,28 @@ void UPixelStreaming2PluginSettings::PostEditChangeProperty(FPropertyChangedEven
 {
 	Super::PostEditChangeProperty(PropertyChangedEvent);
 
-	// NOTE due to legacy variables from the commandline the CVars for settings enums store their string values
-	if (const FEnumProperty* EnumProperty = CastField<const FEnumProperty>(PropertyChangedEvent.Property))
-	{
-		// Enums are not directly alligned in the property map so we get the address
-		FNumericProperty* UnderlyingProp = EnumProperty->GetUnderlyingProperty();
-		void*			  PropertyAddress = EnumProperty->ContainerPtrToValuePtr<void>(this);
-
-		// Get the string value of the change property
-		FString ConsoleString = EnumProperty->GetEnum()->GetNameStringByValue(UnderlyingProp->GetSignedIntPropertyValue(PropertyAddress));
+	FString PropertyName = PropertyChangedEvent.Property->GetNameCPP();
 
-		// If the property has MappedConsoleVariable metadata fetch it and set it to the enums string value
-		if (IConsoleVariable* ConsoleVariable = IConsoleManager::Get().FindConsoleVariable(*PropertyChangedEvent.Property->GetMetaData(PixelStreaming2MappedConsoleVariableFName)))
-		{
-			ConsoleVariable->Set(*ConsoleString, ECVF_SetByProjectSetting);
-		}
-	}
-	else
+	FString CVarName;
+	if (CVarName = FindCVarFromProperty(GetCmdArg, PropertyName); !CVarName.IsEmpty())
 	{
-		if (PropertyChangedEvent.Property->HasMetaData("Bitmask"))
+		if (PropertyName == "WebRTCPortAllocatorFlags")
 		{
-			if (PropertyChangedEvent.Property->GetNameCPP() == "WebRTCPortAllocatorFlags")
-			{
-				SetPortAllocationCVarFromProperty(this, PropertyChangedEvent.Property);
-			}
+			SetPortAllocationCVarFromProperty(this, PropertyChangedEvent.Property);
 		}
-		else if (PropertyChangedEvent.Property->GetNameCPP() == "WebRTCCodecPreferences")
+		else if (PropertyName == "Codec" || PropertyName == "ScalabilityMode" || PropertyName == "EnableSimulcast")
 		{
-			IConsoleVariable* ConsoleVariable = IConsoleManager::Get().FindConsoleVariable(*PropertyChangedEvent.Property->GetMetaData(PixelStreaming2ConsoleVariableMetaFName));
-			VerifyCodecPreferenceSettings(this, PropertyChangedEvent.Property, ConsoleVariable);
+			VerifyVideoSettings();
 		}
-		// Codec and ScalabilityMode properties are updated in VerifyVideoSettings once we know all the settings are compatible
-		else if (PropertyChangedEvent.Property->GetNameCPP() != "Codec" && PropertyChangedEvent.Property->GetNameCPP() != "ScalabilityMode")
+		else
 		{
-			ExportValuesToConsoleVariables(PropertyChangedEvent.Property);
+			SetCVarFromProperty(CVarName, PropertyChangedEvent.Property);
 		}
 	}
-
-	VerifyVideoSettings();
+	else if (CVarName = FindCVarFromProperty(GetMappedCmdArg, PropertyName); !CVarName.IsEmpty())
+	{
+		SetCVarFromProperty(CVarName, PropertyChangedEvent.Property);
+	}
 }
 
 void UPixelStreaming2PluginSettings::VerifyVideoSettings()
@@ -808,20 +873,26 @@ void UPixelStreaming2PluginSettings::VerifyVideoSettings()
 		ScalabilityModeStrProperty->SetPropertyValue_InContainer(this, TEXT("L1T1"));
 	}
 
-	ExportValuesToConsoleVariables(CodecProperty);
-	ExportValuesToConsoleVariables(ScalabilityModeProperty);
+	SetCVarFromProperty(FindCVarFromProperty(GetCmdArg, SimulcastProperty->GetNameCPP()), SimulcastProperty);
+	SetCVarFromProperty(FindCVarFromProperty(GetCmdArg, CodecProperty->GetNameCPP()), CodecProperty);
+	SetCVarFromProperty(FindCVarFromProperty(GetCmdArg, ScalabilityModeProperty->GetNameCPP()), ScalabilityModeProperty);
 }
 #endif
 
-void UPixelStreaming2PluginSettings::SetCVarFromPropertyAndValue(IConsoleVariable* CVar, FProperty* Property, const FString& CVarString, const FString& Value)
+void UPixelStreaming2PluginSettings::SetCVarAndPropertyFromValue(const FString& CVarName, FProperty* Property, const FString& Value)
 {
-	checkf(CVar, TEXT("CVar is nullptr"));
+	IConsoleVariable* CVar = IConsoleManager::Get().FindConsoleVariable(*CVarName);
+	if (!CVar)
+	{
+		UE_LOGFMT(LogPixelStreaming2Settings, Warning, "Failed to find CVar: {0}", CVarName);
+		return;
+	}
 
 	if (FByteProperty* ByteProperty = CastField<FByteProperty>(Property); ByteProperty != NULL && ByteProperty->Enum != NULL)
 	{
-		int32 CastValue;
-		FParse::Value(FCommandLine::Get(), *ConsoleVariableToCommandArgValue(CVarString), CastValue);
-		CVar->Set(CastValue, ECVF_SetByCommandline);
+		CVar->Set(FCString::Atoi(*Value), ECVF_SetByCommandline);
+		ByteProperty->SetPropertyValue_InContainer(this, FCString::Atoi(*Value));
+		UE_LOGFMT(LogPixelStreaming2Settings, Log, "Setting CVar [{0}] and Property [{1}] to [{2}] from command line", CVarName, Property->GetNameCPP(), FCString::Atoi(*Value));
 	}
 	else if (FEnumProperty* EnumProperty = CastField<FEnumProperty>(Property))
 	{
@@ -829,170 +900,157 @@ void UPixelStreaming2PluginSettings::SetCVarFromPropertyAndValue(IConsoleVariabl
 		if (EnumIndex != INDEX_NONE)
 		{
 			CVar->Set(*EnumProperty->GetEnum()->GetNameStringByIndex(EnumIndex), ECVF_SetByCommandline);
+
+			FNumericProperty* UnderlyingProp = EnumProperty->GetUnderlyingProperty();
+			int64*			  PropertyAddress = EnumProperty->ContainerPtrToValuePtr<int64>(this);
+			*PropertyAddress = EnumProperty->GetEnum()->GetValueByIndex(EnumIndex);
+
+			UE_LOGFMT(LogPixelStreaming2Settings, Log, "Setting CVar [{0}] and Property [{1}] to [{2}] from command line", CVarName, Property->GetNameCPP(), EnumProperty->GetEnum()->GetNameStringByIndex(EnumIndex));
 		}
 		else
 		{
-			UE_LOGFMT(LogPixelStreaming2Settings, Warning, "{0} is not a valid enum value for {1}", Value, CVarString);
+			UE_LOGFMT(LogPixelStreaming2Settings, Warning, "{0} is not a valid enum value for {1}", Value, EnumProperty->GetEnum()->CppType);
 		}
 	}
 	else if (FBoolProperty* BoolProperty = CastField<FBoolProperty>(Property))
 	{
+		bool bValue = false;
 		if (Value.Equals(FString(TEXT("true")), ESearchCase::IgnoreCase))
 		{
-			CVar->Set(true, ECVF_SetByCommandline);
+			bValue = true;
 		}
 		else if (Value.Equals(FString(TEXT("false")), ESearchCase::IgnoreCase))
 		{
-			CVar->Set(false, ECVF_SetByCommandline);
-		}
-		else if (FParse::Param(FCommandLine::Get(), *ConsoleVariableToCommandArgParam(CVarString)))
-		{
-			CVar->Set(true, ECVF_SetByCommandline);
+			bValue = false;
 		}
+		CVar->Set(bValue, ECVF_SetByCommandline);
+		BoolProperty->SetPropertyValue_InContainer(this, bValue);
+		UE_LOGFMT(LogPixelStreaming2Settings, Log, "Setting CVar [{0}] and Property [{1}] to [{2}] from command line", CVarName, Property->GetNameCPP(), bValue);
 	}
 	else if (FIntProperty* IntProperty = CastField<FIntProperty>(Property))
 	{
-		int32 CastValue;
-		FParse::Value(FCommandLine::Get(), *ConsoleVariableToCommandArgValue(CVarString), CastValue);
-		CVar->Set(CastValue, ECVF_SetByCommandline);
+		CVar->Set(FCString::Atoi(*Value), ECVF_SetByCommandline);
+		IntProperty->SetPropertyValue_InContainer(this, FCString::Atoi(*Value));
+		UE_LOGFMT(LogPixelStreaming2Settings, Log, "Setting CVar [{0}] and Property [{1}] to [{2}] from command line", CVarName, Property->GetNameCPP(), FCString::Atoi(*Value));
 	}
 	else if (FFloatProperty* FloatProperty = CastField<FFloatProperty>(Property))
 	{
-		float CastValue;
-		FParse::Value(FCommandLine::Get(), *ConsoleVariableToCommandArgValue(CVarString), CastValue);
-		CVar->Set(CastValue, ECVF_SetByCommandline);
+		CVar->Set(FCString::Atof(*Value), ECVF_SetByCommandline);
+		FloatProperty->SetPropertyValue_InContainer(this, FCString::Atof(*Value));
+		UE_LOGFMT(LogPixelStreaming2Settings, Log, "Setting CVar [{0}] and Property [{1}] to [{2}] from command line", CVarName, Property->GetNameCPP(), FCString::Atof(*Value));
 	}
 	else if (FStrProperty* StringProperty = CastField<FStrProperty>(Property))
 	{
 		CVar->Set(*Value, ECVF_SetByCommandline);
+		StringProperty->SetPropertyValue_InContainer(this, *Value);
+		UE_LOGFMT(LogPixelStreaming2Settings, Log, "Setting CVar [{0}] and Property [{1}] to [\"{2}\"] from command line", CVarName, Property->GetNameCPP(), Value);
 	}
 	else if (FNameProperty* NameProperty = CastField<FNameProperty>(Property))
 	{
 		CVar->Set(*Value, ECVF_SetByCommandline);
+		NameProperty->SetPropertyValue_InContainer(this, FName(*Value));
+		UE_LOGFMT(LogPixelStreaming2Settings, Log, "Setting CVar [{0}] and Property [{1}] to [\"{2}\"] from command line", CVarName, Property->GetNameCPP(), Value);
+	}
+	else if (FArrayProperty* ArrayProperty = CastField<FArrayProperty>(Property))
+	{
+		// TODO (william.belcher): Only FString array properties are currently supported
+		CVar->Set(*Value, ECVF_SetByCommandline);
+
+		TArray<FString> StringArray;
+		Value.ParseIntoArray(StringArray, TEXT(","), true);
+
+		TArray<FString>& Array = *ArrayProperty->ContainerPtrToValuePtr<TArray<FString>>(this);
+		Array = StringArray;
+
+		UE_LOGFMT(LogPixelStreaming2Settings, Log, "Setting CVar [{0}] and Property [{1}] to [\"{2}\"] from command line", CVarName, Property->GetNameCPP(), Value);
 	}
 }
 
-void UPixelStreaming2PluginSettings::SetCVarFromProperty(IConsoleVariable* CVar, FProperty* Property, const FString& CVarString)
+void UPixelStreaming2PluginSettings::SetCVarFromProperty(const FString& CVarName, FProperty* Property)
 {
-	checkf(CVar, TEXT("CVar is nullptr"));
+	IConsoleVariable* CVar = IConsoleManager::Get().FindConsoleVariable(*CVarName);
+	if (!CVar)
+	{
+		UE_LOGFMT(LogPixelStreaming2Settings, Warning, "Failed to find CVar: {0}", CVarName);
+		return;
+	}
 
 	if (FByteProperty* ByteProperty = CastField<FByteProperty>(Property); ByteProperty != NULL && ByteProperty->Enum != NULL)
 	{
-		CVar->Set(ByteProperty->GetPropertyValue_InContainer(this), ECVF_SetByCommandline);
+		CVar->Set(ByteProperty->GetPropertyValue_InContainer(this), ECVF_SetByProjectSetting);
+		UE_LOGFMT(LogPixelStreaming2Settings, Log, "Setting CVar [{0}] to [{1}] from Property [{2}]", CVarName, ByteProperty->GetPropertyValue_InContainer(this), Property->GetNameCPP());
 	}
 	else if (FEnumProperty* EnumProperty = CastField<FEnumProperty>(Property))
 	{
 		void* PropertyAddress = EnumProperty->ContainerPtrToValuePtr<void>(this);
 		int64 CurrentValue = EnumProperty->GetUnderlyingProperty()->GetSignedIntPropertyValue(PropertyAddress);
-		CVar->Set(*EnumProperty->GetEnum()->GetNameStringByValue(CurrentValue), ECVF_SetByCommandline);
+		CVar->Set(*EnumProperty->GetEnum()->GetNameStringByValue(CurrentValue), ECVF_SetByProjectSetting);
+		UE_LOGFMT(LogPixelStreaming2Settings, Log, "Setting CVar [{0}] to [{1}] from Property [{2}]", CVarName, EnumProperty->GetEnum()->GetNameStringByValue(CurrentValue), Property->GetNameCPP());
 	}
 	else if (FBoolProperty* BoolProperty = CastField<FBoolProperty>(Property))
 	{
-		CVar->Set(BoolProperty->GetPropertyValue_InContainer(this), ECVF_SetByCommandline);
+		CVar->Set(BoolProperty->GetPropertyValue_InContainer(this), ECVF_SetByProjectSetting);
+		UE_LOGFMT(LogPixelStreaming2Settings, Log, "Setting CVar [{0}] to [{1}] from Property [{2}]", CVarName, BoolProperty->GetPropertyValue_InContainer(this), Property->GetNameCPP());
 	}
 	else if (FIntProperty* IntProperty = CastField<FIntProperty>(Property))
 	{
-		CVar->Set(IntProperty->GetPropertyValue_InContainer(this), ECVF_SetByCommandline);
+		CVar->Set(IntProperty->GetPropertyValue_InContainer(this), ECVF_SetByProjectSetting);
+		UE_LOGFMT(LogPixelStreaming2Settings, Log, "Setting CVar [{0}] to [{1}] from Property [{2}]", CVarName, IntProperty->GetPropertyValue_InContainer(this), Property->GetNameCPP());
 	}
 	else if (FFloatProperty* FloatProperty = CastField<FFloatProperty>(Property))
 	{
-		CVar->Set(FloatProperty->GetPropertyValue_InContainer(this), ECVF_SetByCommandline);
+		CVar->Set(FloatProperty->GetPropertyValue_InContainer(this), ECVF_SetByProjectSetting);
+		UE_LOGFMT(LogPixelStreaming2Settings, Log, "Setting CVar [{0}] to [{1}] from Property [{2}]", CVarName, FloatProperty->GetPropertyValue_InContainer(this), Property->GetNameCPP());
 	}
 	else if (FStrProperty* StringProperty = CastField<FStrProperty>(Property))
 	{
-		CVar->Set(*StringProperty->GetPropertyValue_InContainer(this), ECVF_SetByCommandline);
+		CVar->Set(*StringProperty->GetPropertyValue_InContainer(this), ECVF_SetByProjectSetting);
+		UE_LOGFMT(LogPixelStreaming2Settings, Log, "Setting CVar [{0}] to [\"{1}\"] from Property [{2}]", CVarName, StringProperty->GetPropertyValue_InContainer(this), Property->GetNameCPP());
 	}
 	else if (FNameProperty* NameProperty = CastField<FNameProperty>(Property))
 	{
-		CVar->Set(*NameProperty->GetPropertyValue_InContainer(this).ToString(), ECVF_SetByCommandline);
+		CVar->Set(*NameProperty->GetPropertyValue_InContainer(this).ToString(), ECVF_SetByProjectSetting);
+		UE_LOGFMT(LogPixelStreaming2Settings, Log, "Setting CVar [{0}] to [\"{1}\"] from Property [{2}]", CVarName, NameProperty->GetPropertyValue_InContainer(this), Property->GetNameCPP());
+	}
+	else if (FArrayProperty* ArrayProperty = CastField<FArrayProperty>(Property))
+	{
+		// TODO (william.belcher): Only FString array properties are currently supported
+		TArray<FString> Array = *ArrayProperty->ContainerPtrToValuePtr<TArray<FString>>(this);
+		CVar->Set(*FString::Join(Array, TEXT(",")), ECVF_SetByProjectSetting);
+		UE_LOGFMT(LogPixelStreaming2Settings, Log, "Setting CVar [{0}] to [\"{1}\"] from Property [{2}]", CVarName, FString::Join(Array, TEXT(",")), Property->GetNameCPP());
 	}
 }
 
-// Map of Property Names to their commandline args as GetMetaData() is not avaliable in packaged projects
-static const TMap<FString, FString> GetCmdArg = {
-	{ "LogStats", "PixelStreaming2.LogStats" },
-	{ "SendPlayerIdAsInteger", "PixelStreaming2.SendPlayerIdAsInteger" },
-	{ "DisableLatencyTester", "PixelStreaming2.DisableLatencyTester" },
-	{ "DecoupleFramerate", "PixelStreaming2.DecoupleFrameRate" },
-	{ "DecoupleWaitFactor", "PixelStreaming2.DecoupleWaitFactor" },
-	{ "SignalingReconnectInterval", "PixelStreaming2.SignalingReconnectInterval" },
-	{ "SignalingKeepAliveInterval", "PixelStreaming2.SignalingKeepAliveInterval" },
-	{ "UseMediaCapture", "PixelStreaming2.UseMediaCapture" },
-	{ "DefaultStreamerID", "PixelStreaming2.ID" },
-	{ "ConnectionURL", "PixelStreaming2.ConnectionURL" },
-	{ "CaptureUseFence", "PixelStreaming2.CaptureUseFence" },
-	{ "Codec", "PixelStreaming2.Encoder.Codec" },
-	{ "EncoderTargetBitrate", "PixelStreaming2.Encoder.TargetBitrate" },
-	{ "EncoderMinQuality", "PixelStreaming2.Encoder.MinQuality" },
-	{ "EncoderMaxQuality", "PixelStreaming2.Encoder.MaxQuality" },
-	{ "ScalabilityMode", "PixelStreaming2.Encoder.ScalabilityMode" },
-	{ "KeyframeInterval", "PixelStreaming2.Encoder.KeyframeInterval" },
-	{ "MaxSessions", "PixelStreaming2.Encoder.MaxSessions" },
-	{ "EnableSimulcast", "PixelStreaming2.Encoder.EnableSimulcast" },
-	{ "WebRTCFps", "PixelStreaming2.WebRTC.Fps" },
-	{ "WebRTCStartBitrate", "PixelStreaming2.WebRTC.StartBitrate" },
-	{ "WebRTCMinBitrate", "PixelStreaming2.WebRTC.MinBitrate" },
-	{ "WebRTCMaxBitrate", "PixelStreaming2.WebRTC.MaxBitrate" },
-	{ "WebRTCDisableReceiveAudio", "PixelStreaming2.WebRTC.DisableReceiveAudio" },
-	{ "WebRTCDisableReceiveVideo", "PixelStreaming2.WebRTC.DisableReceiveVideo" },
-	{ "WebRTCDisableTransmitAudio", "PixelStreaming2.WebRTC.DisableTransmitAudio" },
-	{ "WebRTCDisableTransmitVideo", "PixelStreaming2.WebRTC.DisableTransmitVideo" },
-	{ "WebRTCDisableAudioSync", "PixelStreaming2.WebRTC.DisableAudioSync" },
-	{ "WebRTCEnableFlexFec", "PixelStreaming2.WebRTC.EnableFlexFec" },
-	{ "WebRTCDisableStats", "PixelStreaming2.WebRTC.DisableStats" },
-	{ "WebRTCStatsInterval", "PixelStreaming2.WebRTC.StatsInterval" },
-	{ "WebRTCNegotiateCodecs", "PixelStreaming2.WebRTC.NegotiateCodecs" },
-	{ "WebRTCAudioGain", "PixelStreaming2.WebRTC.AudioGain" },
-	{ "WebRTCPortAllocatorFlags", "PixelStreaming2.WebRTC.PortAllocatorFlags" },
-	{ "WebRTCMinPort", "PixelStreaming2.WebRTC.MinPort" },
-	{ "WebRTCMaxPort", "PixelStreaming2.WebRTC.MaxPort" },
-	{ "WebRTCFieldTrials", "PixelStreaming2.WebRTC.FieldTrials" },
-	{ "WebRTCDisableFrameDropper", "PixelStreaming2.WebRTC.DisableFrameDropper" },
-	{ "WebRTCVideoPacingMaxDelay", "PixelStreaming2.WebRTC.VideoPacing.MaxDelay" },
-	{ "WebRTCVideoPacingFactor", "PixelStreaming2.WebRTC.VideoPacing.Factor" },
-	{ "EditorStartOnLaunch", "PixelStreaming2.Editor.StartOnLaunch" },
-	{ "EditorUseRemoteSignallingServer", "PixelStreaming2.Editor.UseRemoteSignallingServer" },
-	{ "HMDEnable", "PixelStreaming2.HMD.Enable" },
-	{ "HMDMatchAspectRatio", "PixelStreaming2.HMD.MatchAspectRatio" },
-	{ "HMDAppleEyePosition", "PixelStreaming2.HMD.ApplyEyePosition" },
-	{ "HMDApplyEyeRotation", "PixelStreaming2.HMD.ApplyEyeRotation" },
-	{ "HMDHFOV", "PixelStreaming2.HMD.HFOV" },
-	{ "HMDVFOV", "PixelStreaming2.HMD.VFOV" },
-	{ "HMDIPD", "PixelStreaming2.HMD.IPD" },
-	{ "HMDProjectionOffsetX", "PixelStreaming2.HMD.ProjectionOffsetX" },
-	{ "HMDProjectionOffsetY", "PixelStreaming2.HMD.ProjectionOffsetY" },
-	{ "InputAllowConsoleCommands", "PixelStreaming2.AllowPixelStreamingCommands" },
-	{ "InputKeyFilter", "PixelStreaming2.KeyFilter" }
-};
-
-static const TMap<FString, FString> GetMappedCmdArg = {
-	{ "InputController", "PixelStreaming2.InputController" },
-	{ "QualityPreset", "PixelStreaming2.Encoder.QualityPreset" },
-	{ "LatencyMode", "PixelStreaming2.Encoder.LatencyMode" },
-	{ "H264Profile", "PixelStreaming2.Encoder.H264Profile" },
-	{ "EditorSource", "PixelStreaming2.Editor.Source" }
-};
-
-//                                                                   Property,   CVar,              CommandLine value
-using FMappingFunc = TFunction<void(UPixelStreaming2PluginSettings*, FProperty*, IConsoleVariable*, const FString&)>;
+void UPixelStreaming2PluginSettings::InitializeCVarsFromProperties()
+{
+	UE_LOGFMT(LogPixelStreaming2Settings, Log, "Initializing CVars from ini");
+	for (FProperty* Property = GetClass()->PropertyLink; Property; Property = Property->PropertyLinkNext)
+	{
+		if (!Property->HasAnyPropertyFlags(CPF_Config))
+		{
+			continue;
+		}
 
-static const TMap<FString, TPair<FString, FMappingFunc>> GetCustomMappedCmdArg = {
-	{ "WebRTCCodecPreferences", { "PixelStreaming2.WebRTC.CodecPreferences", VerifyCodecPreferenceSettings } },
-};
+		// Handle the majority of commandline argument
+		if (Property->GetNameCPP() == "WebRTCPortAllocatorFlags")
+		{
+			SetPortAllocationCVarFromProperty(this, Property);
+			continue;
+		}
 
-static const TArray<FString> GetLegacyCmdArg = {
-	"PixelStreaming2.Encoder.MinQp", // Renamed to MaxQuality
-	"PixelStreaming2.Encoder.MaxQp", // Renamed to MinQuality
-	"PixelStreaming2.IP",			 // Moved to URL
-	"PixelStreaming2.Port",			 // Moved to URL
-	"PixelStreaming2.URL",			 // Renamed to ConnectionURL
-	"PixelStreaming2.SignallingURL", // Renamed to ConnectionURL
-	"AllowPixelStreamingCommands",
-	"PixelStreaming2.NegotiateCodecs", // Renamed to WebRTC.NegotiateCodecs
-	"PixelStreaming2.OnScreenStats",   // CVar is removed but launch arg is used in stats.cpp
-	"PixelStreaming2.HudStats",		   // CVar is removed but launch arg is used in stats.cpp
-	"PixelStreaming2.EnableHMD"		   // Renamed to HMDEnable
-};
+		FString CVarName;
+		if (CVarName = FindCVarFromProperty(GetCmdArg, Property->GetNameCPP()); !CVarName.IsEmpty())
+		{
+			SetCVarFromProperty(CVarName, Property);
+			continue;
+		}
+		else if (CVarName = FindCVarFromProperty(GetMappedCmdArg, Property->GetNameCPP()); !CVarName.IsEmpty())
+		{
+			SetCVarFromProperty(CVarName, Property);
+			continue;
+		}
+	}
+}
 
 void UPixelStreaming2PluginSettings::ValidateCommandLineArgs()
 {
@@ -1019,10 +1077,11 @@ void UPixelStreaming2PluginSettings::ValidateCommandLineArgs()
 		bool bValidArg = false;
 		for (const TPair<FString, FString>& Pair : GetCmdArg)
 		{
-			FString ValidCommandLineArg = ConsoleVariableToCommandArgParam(Pair.Value);
+			FString ValidCommandLineArg = ConsoleVariableToCommandArgParam(Pair.Key);
 			if (CurrentCommandLineArg == ValidCommandLineArg)
 			{
 				bValidArg = true;
+				break;
 			}
 		}
 
@@ -1030,34 +1089,24 @@ void UPixelStreaming2PluginSettings::ValidateCommandLineArgs()
 		{
 			for (const TPair<FString, FString>& Pair : GetMappedCmdArg)
 			{
-				FString ValidCommandLineArg = ConsoleVariableToCommandArgParam(Pair.Value);
+				FString ValidCommandLineArg = ConsoleVariableToCommandArgParam(Pair.Key);
 				if (CurrentCommandLineArg == ValidCommandLineArg)
 				{
 					bValidArg = true;
+					break;
 				}
 			}
 		}
 
 		if (!bValidArg)
 		{
-			for (const TPair<FString, TPair<FString, FMappingFunc>>& Pair : GetCustomMappedCmdArg)
+			for (const TPair<FString, FString>& Pair : GetLegacyCmdArg)
 			{
-				FString ValidCommandLineArg = ConsoleVariableToCommandArgParam(Pair.Value.Key);
-				if (CurrentCommandLineArg == ValidCommandLineArg)
-				{
-					bValidArg = true;
-				}
-			}
-		}
-
-		if (!bValidArg)
-		{
-			for (const FString& LegacyArg : GetLegacyCmdArg)
-			{
-				FString ValidCommandLineArg = ConsoleVariableToCommandArgParam(LegacyArg);
+				FString ValidCommandLineArg = ConsoleVariableToCommandArgParam(Pair.Key);
 				if (CurrentCommandLineArg == ValidCommandLineArg)
 				{
 					bValidArg = true;
+					break;
 				}
 			}
 		}
@@ -1069,178 +1118,172 @@ void UPixelStreaming2PluginSettings::ValidateCommandLineArgs()
 	}
 }
 
-void UPixelStreaming2PluginSettings::ParseLegacyCommandlineArgs()
+void UPixelStreaming2PluginSettings::ParseCommandlineArgs()
 {
-	// Begin legacy PixelStreaming command line args
-	int32 MinQP;
-	if (FParse::Value(FCommandLine::Get(), TEXT("PixelStreamingEncoderMinQp="), MinQP))
+	UE_LOGFMT(LogPixelStreaming2Settings, Verbose, "Updating CVars and properties with command line args");
+	for (const TPair<FString, FString>& Pair : GetCmdArg)
 	{
-		CVarEncoderMaxQuality.AsVariable()->Set(100.0f * (1.0f - (FMath::Clamp<int32>(MinQP, 0, 51) / 51.0f)), ECVF_SetByCommandline);
-		UE_LOGFMT(LogPixelStreaming2Settings, Log, "PixelStreamingEncoderMinQp is a legacy setting, converted to PixelStreamingEncoderMaxQuality={0}", CVarEncoderMaxQuality.GetValueOnAnyThread());
-	}
+		FString CVarString = Pair.Key;
+		FString PropertyName = Pair.Value;
 
-	int32 MaxQP;
-	if (FParse::Value(FCommandLine::Get(), TEXT("PixelStreamingEncoderMaxQp="), MaxQP))
-	{
-		CVarEncoderMinQuality.AsVariable()->Set(100.0f * (1.0f - (FMath::Clamp<int32>(MaxQP, 0, 51) / 51.0f)), ECVF_SetByCommandline);
-		UE_LOGFMT(LogPixelStreaming2Settings, Log, "PixelStreamingEncoderMaxQp is a legacy setting, converted to PixelStreamingEncoderMinQuality={0}", CVarEncoderMinQuality.GetValueOnAnyThread());
-	}
+		FProperty* Property = GetClass()->FindPropertyByName(FName(*PropertyName));
+		if (!Property || !Property->HasAnyPropertyFlags(CPF_Config))
+		{
+			continue;
+		}
 
-	FString LegacyUrl;
-	FString SignallingServerIP;
-	if (FParse::Value(FCommandLine::Get(), TEXT("PixelStreamingIP="), SignallingServerIP))
-	{
-		LegacyUrl += SignallingServerIP;
-	}
+		if (PropertyName == "WebRTCPortAllocatorFlags")
+		{
+			FString ConsoleString;
+			if (FParse::Value(FCommandLine::Get(), *ConsoleVariableToCommandArgValue(CVarString), ConsoleString))
+			{
+				SetPortAllocationCVarAndPropertyFromValue(this, Property, ConsoleString);
+			}
+			continue;
+		}
 
-	FString SignallingServerPort;
-	if (FParse::Value(FCommandLine::Get(), TEXT("PixelStreamingPort="), SignallingServerPort))
-	{
-		LegacyUrl += TEXT(":") + SignallingServerPort;
+		// Handle a directly parsable commandline
+		FString ConsoleString;
+		if (FParse::Value(FCommandLine::Get(), *ConsoleVariableToCommandArgValue(CVarString), ConsoleString))
+		{
+			SetCVarAndPropertyFromValue(CVarString, Property, ConsoleString);
+		}
+		else if (FParse::Param(FCommandLine::Get(), *ConsoleVariableToCommandArgParam(CVarString)))
+		{
+			SetCVarAndPropertyFromValue(CVarString, Property, TEXT("true"));
+		}
 	}
 
-	if (!LegacyUrl.IsEmpty())
+	for (const TPair<FString, FString>& Pair : GetMappedCmdArg)
 	{
-		CVarConnectionURL.AsVariable()->Set(*(FString("ws://") + LegacyUrl), ECVF_SetByCommandline);
-		UE_LOGFMT(LogPixelStreaming2Settings, Log, "PixelStreamingIP and PixelStreamingPort are legacy settings converted to -PixelStreamingConnectionURL={0}", CVarConnectionURL.GetValueOnAnyThread());
-	}
+		FString CVarString = Pair.Key;
+		FString PropertyName = Pair.Value;
 
-	// The new URL argument is PixelStreamingConnectionURL= but we want to support the old one too
-	if (FParse::Value(FCommandLine::Get(), TEXT("PixelStreamingURL="), LegacyUrl))
-	{
-		CVarConnectionURL.AsVariable()->Set(*LegacyUrl, ECVF_SetByCommandline);
-	}
+		FProperty* Property = GetClass()->FindPropertyByName(FName(*PropertyName));
+		if (!Property || !Property->HasAnyPropertyFlags(CPF_Config))
+		{
+			continue;
+		}
 
-	// The new URL argument is PixelStreamingConnectionURL= but we want to support the old one too
-	if (FParse::Value(FCommandLine::Get(), TEXT("PixelStreamingSignallingURL="), LegacyUrl))
-	{
-		CVarConnectionURL.AsVariable()->Set(*LegacyUrl, ECVF_SetByCommandline);
+		// Handle a directly parsable commandline
+		FString ConsoleString;
+		if (FParse::Value(FCommandLine::Get(), *ConsoleVariableToCommandArgValue(CVarString), ConsoleString))
+		{
+			SetCVarAndPropertyFromValue(CVarString, Property, ConsoleString);
+		}
 	}
-
-	ParseLegacyCommandLineOption(TEXT("PixelStreamingNegotiateCodecs"), CVarWebRTCNegotiateCodecs);
-	ParseLegacyCommandLineOption(TEXT("AllowPixelStreamingCommands"), CVarInputAllowConsoleCommands);
-	ParseLegacyCommandLineOption(TEXT("PixelStreamingDebugDumpFrame"), CVarEncoderDebugDumpFrame);
-	// End legacy PixelStreaming command line args
-
-	// Begin legacy PixelStreamingEditor command line args
-	ParseLegacyCommandLineOption(TEXT("EditorPixelStreamingStartOnLaunch"), CVarEditorStartOnLaunch);
-	ParseLegacyCommandLineOption(TEXT("EditorPixelStreamingUseRemoteSignallingServer"), CVarEditorUseRemoteSignallingServer);
-
-	ParseLegacyCommandLineValue(TEXT("EditorPixelStreamingSource="), CVarEditorSource);
-	IConsoleVariable* EditorSourceCVar = IConsoleManager::Get().FindConsoleVariable(TEXT("PixelStreaming2.Editor.Source"));
-	CheckConsoleEnum<EPixelStreaming2EditorStreamTypes>(EditorSourceCVar);
-	// End legacy PixelStreamingEditor command line args
-
-	// End legacy PixelStreamingHMD command line args
-	ParseLegacyCommandLineOption(TEXT("PixelStreamingEnableHMD"), CVarHMDEnable);
-	// End legacy PixelStreamingHMD command line args
 }
 
-void UPixelStreaming2PluginSettings::PostInitProperties()
+void UPixelStreaming2PluginSettings::ParseLegacyCommandlineArgs()
 {
-	Super::PostInitProperties();
+	FString SignallingServerIP;
+	FString SignallingServerPort;
 
-	UE_LOGFMT(LogPixelStreaming2Settings, Log, "Initialising Pixel Streaming settings.");
+	for (const TPair<FString, FString>& Pair : GetLegacyCmdArg)
+	{
+		FString LegacyCVarString = Pair.Key;
+		FString PropertyName = Pair.Value;
 
-	ValidateCommandLineArgs();
+		FProperty* Property = GetClass()->FindPropertyByName(FName(*PropertyName));
+		if (!Property || !Property->HasAnyPropertyFlags(CPF_Config))
+		{
+			continue;
+		}
 
-	for (FProperty* Property = GetClass()->PropertyLink; Property; Property = Property->PropertyLinkNext)
-	{
-		if (!Property->HasAnyPropertyFlags(CPF_Config))
+		FString NewCVarString;
+		if (FString CmdArgCVar = FindCVarFromProperty(GetCmdArg, PropertyName); !CmdArgCVar.IsEmpty())
+		{
+			NewCVarString = CmdArgCVar;
+		}
+		else if (FString MappedCmdArgCVar = FindCVarFromProperty(GetMappedCmdArg, PropertyName); !MappedCmdArgCVar.IsEmpty())
+		{
+			NewCVarString = MappedCmdArgCVar;
+		}
+		else
 		{
 			continue;
 		}
 
-		// Handle the majority of commandline argument
-		if (GetCmdArg.Contains(Property->GetNameCPP()))
+		if (LegacyCVarString == "PixelStreaming2.Encoder.MinQp")
 		{
-			FString CVarString = GetCmdArg[Property->GetNameCPP()];
-			if (Property->GetNameCPP() == "WebRTCPortAllocatorFlags")
+			int32 MinQP;
+			if (FParse::Value(FCommandLine::Get(), *ConsoleVariableToCommandArgValue(LegacyCVarString), MinQP))
 			{
-				FString ConsoleString;
-				if (FParse::Value(FCommandLine::Get(), *ConsoleVariableToCommandArgValue(CVarString), ConsoleString))
-				{
-					IConsoleManager::Get().FindConsoleVariable(*CVarString)->Set(*ConsoleString, ECVF_SetByCommandline);
-				}
-				else
-				{
-					SetPortAllocationCVarFromProperty(this, Property);
-				}
+				int32 ScaledQuality = 100.0f * (1.0f - (FMath::Clamp<int32>(MinQP, 0, 51) / 51.0f));
+				SetCVarAndPropertyFromValue(NewCVarString, Property, FString::Printf(TEXT("%d"), ScaledQuality));
+				UE_LOGFMT(LogPixelStreaming2Settings, Warning, "PixelStreamingEncoderMinQp is a legacy setting, converted to PixelStreamingEncoderMaxQuality={0}", CVarEncoderMaxQuality.GetValueOnAnyThread());
 				continue;
 			}
-
-			if (!CVarString.IsEmpty())
+		}
+		else if (LegacyCVarString == "PixelStreaming2.Encoder.MaxQp")
+		{
+			int32 MaxQP;
+			if (FParse::Value(FCommandLine::Get(), *ConsoleVariableToCommandArgValue(LegacyCVarString), MaxQP))
 			{
-				// Handle a directly parsable commandline
-				FString ConsoleString;
-				if (FParse::Value(FCommandLine::Get(), *ConsoleVariableToCommandArgValue(CVarString), ConsoleString))
-				{
-					SetCVarFromPropertyAndValue(IConsoleManager::Get().FindConsoleVariable(*CVarString), Property, CVarString, ConsoleString);
-				}
-				else if (FParse::Param(FCommandLine::Get(), *ConsoleVariableToCommandArgParam(CVarString)))
-				{
-					SetCVarFromPropertyAndValue(IConsoleManager::Get().FindConsoleVariable(*CVarString), Property, CVarString, TEXT("true"));
-				}
-				else
-				{
-					SetCVarFromProperty(IConsoleManager::Get().FindConsoleVariable(*CVarString), Property, CVarString);
-				}
+				int32 ScaledQuality = 100.0f * (1.0f - (FMath::Clamp<int32>(MaxQP, 0, 51) / 51.0f));
+				SetCVarAndPropertyFromValue(NewCVarString, Property, FString::Printf(TEXT("%d"), ScaledQuality));
+				UE_LOGFMT(LogPixelStreaming2Settings, Warning, "PixelStreamingEncoderMaxQp is a legacy setting, converted to PixelStreamingEncoderMinQuality={0}", CVarEncoderMinQuality.GetValueOnAnyThread());
+				continue;
 			}
 		}
-
-		// Handle a commandline argument that needs mapping from string to enum string
-		if (GetMappedCmdArg.Contains(Property->GetNameCPP()))
+		else if (LegacyCVarString == "PixelStreaming2.IP" || LegacyCVarString == "PixelStreaming2.Port")
 		{
-			FString CVarString = GetMappedCmdArg[Property->GetNameCPP()];
-			FString ConsoleString;
-			if (FParse::Value(FCommandLine::Get(), *ConsoleVariableToCommandArgValue(CVarString), ConsoleString))
+			if (LegacyCVarString == "PixelStreaming2.IP")
 			{
-				SetCVarFromPropertyAndValue(IConsoleManager::Get().FindConsoleVariable(*CVarString), Property, CVarString, ConsoleString);
+				FParse::Value(FCommandLine::Get(), *ConsoleVariableToCommandArgValue(LegacyCVarString), SignallingServerIP);
 			}
-			else
+			else if (LegacyCVarString == "PixelStreaming2.Port")
 			{
-				// Safety check that it is actually an EnumProperty
-				if (const FEnumProperty* EnumProperty = CastField<const FEnumProperty>(Property))
-				{
-					FNumericProperty* UnderlyingProp = EnumProperty->GetUnderlyingProperty();
-					void*			  PropertyAddress = EnumProperty->ContainerPtrToValuePtr<void>(this);
-					ConsoleString = EnumProperty->GetEnum()->GetNameStringByValue(UnderlyingProp->GetSignedIntPropertyValue(PropertyAddress));
-
-					if (IConsoleVariable* ConsoleVariable = IConsoleManager::Get().FindConsoleVariable(*CVarString))
-					{
-						ConsoleVariable->Set(*ConsoleString, ECVF_SetByProjectSetting);
-					}
-				}
-				else
-				{
-					checkNoEntry();
-				}
+				FParse::Value(FCommandLine::Get(), *ConsoleVariableToCommandArgValue(LegacyCVarString), SignallingServerPort);
 			}
-		}
-
-		// Handle a commandline argument that needs custom mapping from string to some other type
-		// eg TArray to comma separate FString
-		if (GetCustomMappedCmdArg.Contains(Property->GetNameCPP()))
-		{
-			TPair<FString, FMappingFunc> CVarFuncPair = GetCustomMappedCmdArg[Property->GetNameCPP()];
-			FString						 CVarString = CVarFuncPair.Key;
-			FMappingFunc				 MappingFunc = CVarFuncPair.Value;
 
-			FString ConsoleString;
-			if (FParse::Value(FCommandLine::Get(), *ConsoleVariableToCommandArgValue(CVarString), ConsoleString, false))
+			if (!SignallingServerIP.IsEmpty() && !SignallingServerPort.IsEmpty())
 			{
-				// Pass in the console string value. This will set the CVar from what was on the command line
-				MappingFunc(this, Property, IConsoleManager::Get().FindConsoleVariable(*CVarString), ConsoleString);
-			}
-			else
-			{
-				// Pass in an empty value. This will set the CVar from the property's value
-				MappingFunc(this, Property, IConsoleManager::Get().FindConsoleVariable(*CVarString), TEXT(""));
+				FString LegacyUrl = TEXT("ws://") + SignallingServerIP + TEXT(":") + SignallingServerPort;
+				SetCVarAndPropertyFromValue(NewCVarString, Property, LegacyUrl);
+				UE_LOGFMT(LogPixelStreaming2Settings, Warning, "PixelStreamingIP and PixelStreamingPort are legacy settings converted to -PixelStreamingConnectionURL={0}", CVarConnectionURL.GetValueOnAnyThread());
 			}
+
+			continue;
+		}
+
+		FString ConsoleString;
+		if (FParse::Value(FCommandLine::Get(), *ConsoleVariableToCommandArgValue(LegacyCVarString), ConsoleString))
+		{
+			SetCVarAndPropertyFromValue(NewCVarString, Property, ConsoleString);
 		}
+		else if (FParse::Param(FCommandLine::Get(), *ConsoleVariableToCommandArgParam(LegacyCVarString)))
+		{
+			SetCVarAndPropertyFromValue(NewCVarString, Property, TEXT("true"));
+		}
+		else
+		{
+			continue;
+		}
+
+		UE_LOGFMT(LogPixelStreaming2Settings, Warning, "{0} is a legacy setting and has been converted to {1}", ConsoleVariableToCommandArgParam(LegacyCVarString), ConsoleVariableToCommandArgParam(NewCVarString));
 	}
 
-	// Handle parsing of legacy command line args (such as -PixelStreamingUrl) after .ini, properties, and new commandline args.
+	ParseLegacyCommandLineOption(TEXT("PixelStreamingDebugDumpFrame"), CVarEncoderDebugDumpFrame);
+	// End legacy PixelStreaming command line args
+}
+
+void UPixelStreaming2PluginSettings::PostInitProperties()
+{
+	Super::PostInitProperties();
+
+	UE_LOGFMT(LogPixelStreaming2Settings, Log, "Initialising Pixel Streaming settings.");
+
+	// Set all the CVars to reflect the state of the ini
+	InitializeCVarsFromProperties();
+
+	// Validate command line args to log if they're invalid
+	ValidateCommandLineArgs();
+
+	// Update CVars and properties based on command line args
+	ParseCommandlineArgs();
+
+	// Handle parsing of legacy command line args (such as -PixelStreamingUrl) after .ini and new commandline args.
 	ParseLegacyCommandlineArgs();
 }
 
@@ -1300,4 +1343,25 @@ TArray<FString> UPixelStreaming2PluginSettings::GetScalabilityModeOptions() cons
 	}
 
 	return ScalabilityModes;
+}
+
+TArray<FString> UPixelStreaming2PluginSettings::GetWebRTCCodecPreferencesOptions() const
+{
+	TSet<FString> PossibleCodecs = {
+		UE::PixelStreaming2::GetCVarStringFromEnum(EVideoCodec::AV1),
+		UE::PixelStreaming2::GetCVarStringFromEnum(EVideoCodec::H264),
+		UE::PixelStreaming2::GetCVarStringFromEnum(EVideoCodec::VP9),
+		UE::PixelStreaming2::GetCVarStringFromEnum(EVideoCodec::VP8)
+	};
+
+	FProperty*		Property = GetClass()->FindPropertyByName(TEXT("WebRTCCodecPreferences"));
+	FArrayProperty* ArrayProperty = CastField<FArrayProperty>(Property);
+	TArray<FString> CurrentCodecArray = *ArrayProperty->ContainerPtrToValuePtr<TArray<FString>>(this);
+
+	for (FString VideoCodec : CurrentCodecArray)
+	{
+		PossibleCodecs.Remove(VideoCodec);
+	}
+
+	return PossibleCodecs.Array();
 }
\ No newline at end of file
