[PS2] Fix Pixel Streaming 2 VP9 frame corruption.

bAutomaticResizeOn causes sporatic frame corruption. This setting is receivedd from EpirRTC which fetches it from WebRTC.

PixelStreaming 1 does not suffer from thsi issue as it uses the vpx encoder embedded in WebRTC. WebRTC disables it if quality_scaler_experiment_ is enabled, with is enabled by default.

diff --git a/Engine/Plugins/Media/PixelStreaming2/Source/PixelStreaming2RTC/Private/EpicRtcVideoEncoder.cpp b/Engine/Plugins/Media/PixelStreaming2/Source/PixelStreaming2RTC/Private/EpicRtcVideoEncoder.cpp
--- a/Engine/Plugins/Media/PixelStreaming2/Source/PixelStreaming2RTC/Private/EpicRtcVideoEncoder.cpp
+++ b/Engine/Plugins/Media/PixelStreaming2/Source/PixelStreaming2RTC/Private/EpicRtcVideoEncoder.cpp
@@ -194,7 +194,8 @@ namespace UE::PixelStreaming2
 				VideoConfig->NumberOfCores = EncoderConfig._numberOfCores;
 				VideoConfig->bDenoisingOn = static_cast<bool>(EncoderConfig._isDenoisingOn);
 				VideoConfig->bAdaptiveQpMode = static_cast<bool>(EncoderConfig._isAdaptiveQpMode);
-				VideoConfig->bAutomaticResizeOn = static_cast<bool>(EncoderConfig._isAutomaticResizeOn);
+				// TODO RTCP-7994 (Eden.Harris) bAutomaticResizeOn can result in sporadic frame corruption.
+				VideoConfig->bAutomaticResizeOn = false; // static_cast<bool>(EncoderConfig._isAutomaticResizeOn);
 				VideoConfig->bFlexibleMode = static_cast<bool>(EncoderConfig._isFlexibleMode);
 				VideoConfig->InterLayerPrediction = static_cast<UE::AVCodecCore::VP9::EInterLayerPrediction>(EncoderConfig._interLayerPred);
 				for (size_t i = 0; i < EncoderConfig._spatialLayers->Size(); i++)
