[PS2] QOL: Use MediaCaptureCapturer for RHI to I420 conversions

diff --git a/Engine/Plugins/Media/PixelStreaming2/Source/PixelStreaming2/Private/VideoCapturer.cpp b/Engine/Plugins/Media/PixelStreaming2/Source/PixelStreaming2/Private/VideoCapturer.cpp
--- a/Engine/Plugins/Media/PixelStreaming2/Source/PixelStreaming2/Private/VideoCapturer.cpp
+++ b/Engine/Plugins/Media/PixelStreaming2/Source/PixelStreaming2/Private/VideoCapturer.cpp
@@ -170,7 +170,14 @@ namespace UE::PixelStreaming2
 						}
 					}
 				case PixelCaptureBufferFormat::FORMAT_I420:
-					return FPixelCaptureCapturerRHIToI420CPU::Create(FinalScale);
+					if (UPixelStreaming2PluginSettings::CVarUseMediaCapture.GetValueOnAnyThread())
+					{
+						return FPixelCaptureCapturerMediaCapture::Create(FinalScale, FinalFormat);
+					}
+					else
+					{
+						return FPixelCaptureCapturerRHIToI420CPU::Create(FinalScale);
+					}
 				default:
 					UE_LOGFMT(LogPixelStreaming2, Error, "Unsupported final format ({0}) for RHI input format", FinalFormat);
 					return nullptr;
