[PS2] Fix H264 not streaming with NegotiateCodecs with AMD GPU and Pixel Streaming 2.

Problem is with negotiate codecs, we set the Video direction to EpicRtcMediaSourceDirection::SendRecv. Normally with AMD we set it to EpicRtcMediaSourceDirection::SendOnly with H264 since it cannot decode. But since we are negotiating the codec, we don't know if H264 will be selected, or VPX where decoding is supported and therefore SendRecv works.

This fix disables receiving video when AMD gpu and NegotiateCodecs is enabled. Explicitly enabling receiving video overrides this.

diff --git a/Engine/Plugins/Media/PixelStreaming2/Source/PixelStreaming2RTC/Private/EpicRtcStreamer.cpp b/Engine/Plugins/Media/PixelStreaming2/Source/PixelStreaming2RTC/Private/EpicRtcStreamer.cpp
--- a/Engine/Plugins/Media/PixelStreaming2/Source/PixelStreaming2RTC/Private/EpicRtcStreamer.cpp
+++ b/Engine/Plugins/Media/PixelStreaming2/Source/PixelStreaming2RTC/Private/EpicRtcStreamer.cpp
@@ -990,7 +990,8 @@ namespace UE::PixelStreaming2
 		// Check if the user has selected only H.264 on a AMD gpu and disable receiving video.
 		// WebRTC does not support using SendRecv if the encoding and decoding do not support the same codec.
 		// AMD GPUs currently have decoding disabled so WebRTC fails to create SDP codecs with SendRecv.
-		if (IsRHIDeviceAMD() && !bNegotiateCodecs && SelectedCodec == EVideoCodec::H264)
+		// TODO (Eden.Harris) RTCP-8039: This workaround won't be needed once H.264 decoding is enabled with AMD GPUs.
+		if (IsRHIDeviceAMD() && (bNegotiateCodecs || (!bNegotiateCodecs && SelectedCodec == EVideoCodec::H264)))
 		{
 			if (bReceiveBrowserVideo)
 			{
