[PS1, PS2] Fix Pixel Streaming Capture Stats.

Capture times are captured with milliseconds but the stats expected cycles.Fix changes stats to expect milliseconds.

diff --git a/Engine/Plugins/Media/PixelStreaming/Source/PixelStreaming/Private/Stats.cpp b/Engine/Plugins/Media/PixelStreaming/Source/PixelStreaming/Private/Stats.cpp
--- a/Engine/Plugins/Media/PixelStreaming/Source/PixelStreaming/Private/Stats.cpp
+++ b/Engine/Plugins/Media/PixelStreaming/Source/PixelStreaming/Private/Stats.cpp
@@ -611,9 +611,9 @@ namespace UE::PixelStreaming
 		}
 	}
 
-	double FStats::AddTimeStat(uint64 Cycles1, const FString& Label)
+	double FStats::AddTimeStat(uint64 Millis, const FString& Label)
 	{
-		const double DeltaMs = FPlatformTime::ToMilliseconds64(Cycles1);
+		const double	DeltaMs = Millis;
 		const FStatData TimeData{ FName(*Label), DeltaMs, 2, true };
 		StoreApplicationStat(TimeData);
 		return DeltaMs;
diff --git a/Engine/Plugins/Media/PixelStreaming/Source/PixelStreaming/Private/Stats.h b/Engine/Plugins/Media/PixelStreaming/Source/PixelStreaming/Private/Stats.h
--- a/Engine/Plugins/Media/PixelStreaming/Source/PixelStreaming/Private/Stats.h
+++ b/Engine/Plugins/Media/PixelStreaming/Source/PixelStreaming/Private/Stats.h
@@ -160,8 +160,8 @@ namespace UE::PixelStreaming
 
 		FORCEINLINE TStatId GetStatId() const { RETURN_QUICK_DECLARE_CYCLE_STAT(PixelStreamingStats, STATGROUP_Tickables); }
 
-		double AddTimeStat(uint64 Cycles1, const FString& Label);
-		double AddTimeDeltaStat(uint64 Cycles1, uint64 Cycles2, const FString& Label);
+		double AddTimeStat(uint64 Millis, const FString& Label);
+		double AddTimeDeltaStat(uint64 Millis1, uint64 Millis2, const FString& Label);
 		void AddFrameTimingStats(const FPixelCaptureFrameMetadata& FrameMetadata);
 		void AddCanvasTile(FName Name, const FCanvasTileItem& Tile);
 
@@ -174,7 +174,6 @@ namespace UE::PixelStreaming
 		FStats();
 		void RegisterEngineHooks();
 		void PollPixelStreamingSettings();
-		void RemovePeerStat(FPixelStreamingPlayerId PlayerId);
 		void FireStatChanged(FPixelStreamingPlayerId PlayerId, FName StatName, float StatValue);
 		void UpdateConsoleAutoComplete(TArray<FAutoCompleteCommand>& AutoCompleteList);
 		void GraphValue_GameThread(FName InName, float Value, int InSamples, float InMinRange, float InMaxRange, float InRefValue);
diff --git a/Engine/Plugins/Media/PixelStreaming2/Source/PixelStreaming2RTC/Private/Stats.cpp b/Engine/Plugins/Media/PixelStreaming2/Source/PixelStreaming2RTC/Private/Stats.cpp
--- a/Engine/Plugins/Media/PixelStreaming2/Source/PixelStreaming2RTC/Private/Stats.cpp
+++ b/Engine/Plugins/Media/PixelStreaming2/Source/PixelStreaming2RTC/Private/Stats.cpp
@@ -623,9 +623,9 @@ namespace UE::PixelStreaming2
 		}
 	}
 
-	double FStats::AddTimeStat(uint64 Cycles1, const FString& Label)
+	double FStats::AddTimeStat(uint64 Millis, const FString& Label)
 	{
-		const double	DeltaMs = FPlatformTime::ToMilliseconds64(Cycles1);
+		const double	DeltaMs = Millis;
 		const FStatData TimeData{ FName(*Label), DeltaMs, 2, true };
 		StoreApplicationStat(TimeData);
 		return DeltaMs;
diff --git a/Engine/Plugins/Media/PixelStreaming2/Source/PixelStreaming2RTC/Private/Stats.h b/Engine/Plugins/Media/PixelStreaming2/Source/PixelStreaming2RTC/Private/Stats.h
--- a/Engine/Plugins/Media/PixelStreaming2/Source/PixelStreaming2RTC/Private/Stats.h
+++ b/Engine/Plugins/Media/PixelStreaming2/Source/PixelStreaming2RTC/Private/Stats.h
@@ -152,8 +152,8 @@ namespace UE::PixelStreaming2
 
 		FORCEINLINE TStatId GetStatId() const { RETURN_QUICK_DECLARE_CYCLE_STAT(PixelStreaming2Stats, STATGROUP_Tickables); }
 
-		double AddTimeStat(uint64 Cycles1, const FString& Label);
-		double AddTimeDeltaStat(uint64 Cycles1, uint64 Cycles2, const FString& Label);
+		double AddTimeStat(uint64 Millis, const FString& Label);
+		double AddTimeDeltaStat(uint64 Millis1, uint64 Millis2, const FString& Label);
 		void   AddFrameTimingStats(const FPixelCaptureFrameMetadata& FrameMetadata);
 		void   AddCanvasTile(FName Name, const FCanvasTileItem& Tile);
 
@@ -169,7 +169,6 @@ namespace UE::PixelStreaming2
 		void UpdateConsoleAutoComplete(TArray<FAutoCompleteCommand>& AutoCompleteList);
 
 		void PollPixelStreaming2Settings();
-		void RemovePeerStat(const FString& PlayerId);
 		void FireStatChanged(const FString& PlayerId, FName StatName, float StatValue);
 		void GraphValue_GameThread(FName InName, float Value, int InSamples, float InMinRange, float InMaxRange, float InRefValue);
 		void AddCanvasTile_GameThread(FName Name, const FCanvasTileItem& Tile);
