[PS2] Fix: Re-enable directional media tests

diff --git a/Engine/Plugins/Media/PixelStreaming2/Source/PixelStreaming2RTC/Private/Tests/MediaDirectionTest.cpp b/Engine/Plugins/Media/PixelStreaming2/Source/PixelStreaming2RTC/Private/Tests/MediaDirectionTest.cpp
new file mode 100644
--- /dev/null
+++ b/Engine/Plugins/Media/PixelStreaming2/Source/PixelStreaming2RTC/Private/Tests/MediaDirectionTest.cpp
@@ -0,0 +1,195 @@
+// Copyright Epic Games, Inc. All Rights Reserved.
+
+#if WITH_DEV_AUTOMATION_TESTS
+
+	#include "Misc/AutomationTest.h"
+	#include "PixelStreaming2PluginSettings.h"
+	#include "TestUtils.h"
+	#include "UtilsAsync.h"
+
+namespace UE::PixelStreaming2
+{
+	void DoMediaDirectionTest(EMediaType Media, EMediaDirection Direction)
+	{
+		int32 StreamerPort = TestUtils::NextStreamerPort();
+		int32 PlayerPort = TestUtils::NextPlayerPort();
+
+		TSharedPtr<UE::PixelStreaming2Servers::IServer> SignallingServer = CreateSignallingServer(StreamerPort, PlayerPort);
+
+		FString								 StreamerName(FString::Printf(TEXT("MockStreamer%d"), StreamerPort));
+		TSharedPtr<IPixelStreaming2Streamer> Streamer = CreateStreamer(StreamerName, StreamerPort);
+		TSharedPtr<FVideoProducer>			 VideoProducer = FVideoProducer::Create();
+		Streamer->SetVideoProducer(VideoProducer);
+
+		EMediaDirection PlayerMediaDirection = EMediaDirection::Bidirectional;
+		if (Direction == EMediaDirection::SendOnly)
+		{
+			// Streamer is only sending so player should only receive
+			PlayerMediaDirection = EMediaDirection::RecvOnly;
+		}
+		else if (Direction == EMediaDirection::RecvOnly)
+		{
+			// Streamer is only receiving so player should only send
+			PlayerMediaDirection = EMediaDirection::SendOnly;
+		}
+		else if (Direction == EMediaDirection::Disabled)
+		{
+			// Streamer has disabled media so player disables it too
+			PlayerMediaDirection = EMediaDirection::Disabled;
+		}
+
+		FMockPlayerConfig PlayerConfig = { .AudioDirection = EMediaDirection::Disabled, .VideoDirection = EMediaDirection::Disabled };
+		if (Media == EMediaType::Audio)
+		{
+			PlayerConfig.AudioDirection = PlayerMediaDirection;
+		}
+		else if (Media == EMediaType::Video)
+		{
+			PlayerConfig.VideoDirection = PlayerMediaDirection;
+		}
+
+		TSharedPtr<FMockPlayer> Player = CreatePlayer(PlayerConfig);
+
+		Streamer->StartStreaming();
+
+		if (!Streamer->IsStreaming())
+		{
+			UE_LOG(LogPixelStreaming2RTC, Error, TEXT("Streamer is not connected."));
+			return;
+		}
+
+		TSharedPtr<bool> bLocalTrack = MakeShared<bool>(false);
+		TSharedPtr<bool> bRemoteTrack = MakeShared<bool>(false);
+		if (Media == EMediaType::Audio)
+		{
+			Player->EpicRtcManager->OnAudioTrackUpdate.AddLambda([bLocalTrack, bRemoteTrack](EpicRtcParticipantInterface* Particpant, EpicRtcAudioTrackInterface* AudioTrack) {
+				if (AudioTrack->IsRemote())
+				{
+					*bRemoteTrack.Get() = true;
+				}
+				else
+				{
+					*bLocalTrack.Get() = true;
+				}
+			});
+		}
+		else if (Media == EMediaType::Video)
+		{
+			Player->EpicRtcManager->OnVideoTrackUpdate.AddLambda([bLocalTrack, bRemoteTrack](EpicRtcParticipantInterface* Particpant, EpicRtcVideoTrackInterface* VideoTrack) {
+				if (VideoTrack->IsRemote())
+				{
+					*bRemoteTrack.Get() = true;
+				}
+				else
+				{
+					*bLocalTrack.Get() = true;
+				}
+			});
+		}
+
+		Player->Connect(PlayerPort);
+
+		ADD_LATENT_AUTOMATION_COMMAND(FSubscribePlayerAfterStreamerConnectedOrTimeout(5.0, Streamer, Player, StreamerName))
+
+		if (Direction == EMediaDirection::SendOnly)
+		{
+			// Check that the player has only a remote track
+			ADD_LATENT_AUTOMATION_COMMAND(FWaitForPlayerTrackOrTimeout(5.0, Player, bRemoteTrack))
+			ADD_LATENT_AUTOMATION_COMMAND(FWaitAndCheckBool(5.0, Player, bLocalTrack, false))
+		}
+		else if (Direction == EMediaDirection::RecvOnly)
+		{
+			// Check that the player has only a local track
+			ADD_LATENT_AUTOMATION_COMMAND(FWaitForPlayerTrackOrTimeout(5.0, Player, bLocalTrack))
+			ADD_LATENT_AUTOMATION_COMMAND(FWaitAndCheckBool(5.0, Player, bRemoteTrack, false))
+		}
+		else if (Direction == EMediaDirection::Bidirectional)
+		{
+			// Check that the player has both a local and remote track
+			ADD_LATENT_AUTOMATION_COMMAND(FWaitForPlayerTrackOrTimeout(5.0, Player, bLocalTrack))
+			ADD_LATENT_AUTOMATION_COMMAND(FWaitForPlayerTrackOrTimeout(5.0, Player, bRemoteTrack))
+		}
+		else if (Direction == EMediaDirection::Disabled)
+		{
+			// Wait 5 seconds and then check that the player has no local and no remote track
+			ADD_LATENT_AUTOMATION_COMMAND(FWaitAndCheckBool(5.0, Player, bLocalTrack, false))
+			ADD_LATENT_AUTOMATION_COMMAND(FWaitAndCheckBool(5.0, Player, bRemoteTrack, false))
+		}
+
+		ADD_LATENT_AUTOMATION_COMMAND(FCleanupAll(SignallingServer, Streamer, Player))
+	}
+
+	IMPLEMENT_SIMPLE_AUTOMATION_TEST(FPS2AudioSendOnlyTest, "System.Plugins.PixelStreaming2.FPS2AudioSendOnlyTest", EAutomationTestFlags::EditorContext | EAutomationTestFlags::ClientContext | EAutomationTestFlags::ProductFilter)
+	bool FPS2AudioSendOnlyTest::RunTest(const FString& Parameters)
+	{
+		SetMediaDirection(EMediaType::Audio, EMediaDirection::SendOnly);
+		DoMediaDirectionTest(EMediaType::Audio, EMediaDirection::SendOnly);
+
+		return true;
+	}
+
+	IMPLEMENT_SIMPLE_AUTOMATION_TEST(FPS2AudioRecvOnlyTest, "System.Plugins.PixelStreaming2.FPS2AudioRecvOnlyTest", EAutomationTestFlags::EditorContext | EAutomationTestFlags::ClientContext | EAutomationTestFlags::ProductFilter)
+	bool FPS2AudioRecvOnlyTest::RunTest(const FString& Parameters)
+	{
+		SetMediaDirection(EMediaType::Audio, EMediaDirection::RecvOnly);
+		DoMediaDirectionTest(EMediaType::Audio, EMediaDirection::RecvOnly);
+
+		return true;
+	}
+
+	IMPLEMENT_SIMPLE_AUTOMATION_TEST(FPS2AudioDisabledTest, "System.Plugins.PixelStreaming2.FPS2AudioDisabledTest", EAutomationTestFlags::EditorContext | EAutomationTestFlags::ClientContext | EAutomationTestFlags::ProductFilter)
+	bool FPS2AudioDisabledTest::RunTest(const FString& Parameters)
+	{
+		SetMediaDirection(EMediaType::Audio, EMediaDirection::Disabled);
+		DoMediaDirectionTest(EMediaType::Audio, EMediaDirection::Disabled);
+
+		return true;
+	}
+
+	IMPLEMENT_SIMPLE_AUTOMATION_TEST(FPS2AudioBidirectionalTest, "System.Plugins.PixelStreaming2.FPS2AudioBidirectionalTest", EAutomationTestFlags::EditorContext | EAutomationTestFlags::ClientContext | EAutomationTestFlags::ProductFilter)
+	bool FPS2AudioBidirectionalTest::RunTest(const FString& Parameters)
+	{
+		SetMediaDirection(EMediaType::Audio, EMediaDirection::Bidirectional);
+		DoMediaDirectionTest(EMediaType::Audio, EMediaDirection::Bidirectional);
+
+		return true;
+	}
+
+	IMPLEMENT_SIMPLE_AUTOMATION_TEST(FPS2VideoSendOnlyTest, "System.Plugins.PixelStreaming2.FPS2VideoSendOnlyTest", EAutomationTestFlags::EditorContext | EAutomationTestFlags::ClientContext | EAutomationTestFlags::ProductFilter)
+	bool FPS2VideoSendOnlyTest::RunTest(const FString& Parameters)
+	{
+		SetMediaDirection(EMediaType::Video, EMediaDirection::SendOnly);
+		DoMediaDirectionTest(EMediaType::Video, EMediaDirection::SendOnly);
+
+		return true;
+	}
+
+	IMPLEMENT_SIMPLE_AUTOMATION_TEST(FPS2VideoRecvOnlyTest, "System.Plugins.PixelStreaming2.FPS2VideoRecvOnlyTest", EAutomationTestFlags::EditorContext | EAutomationTestFlags::ClientContext | EAutomationTestFlags::ProductFilter)
+	bool FPS2VideoRecvOnlyTest::RunTest(const FString& Parameters)
+	{
+		SetMediaDirection(EMediaType::Video, EMediaDirection::RecvOnly);
+		DoMediaDirectionTest(EMediaType::Video, EMediaDirection::RecvOnly);
+
+		return true;
+	}
+
+	IMPLEMENT_SIMPLE_AUTOMATION_TEST(FPS2VideoDisabledTest, "System.Plugins.PixelStreaming2.FPS2VideoDisabledTest", EAutomationTestFlags::EditorContext | EAutomationTestFlags::ClientContext | EAutomationTestFlags::ProductFilter)
+	bool FPS2VideoDisabledTest::RunTest(const FString& Parameters)
+	{
+		SetMediaDirection(EMediaType::Video, EMediaDirection::Disabled);
+		DoMediaDirectionTest(EMediaType::Video, EMediaDirection::Disabled);
+
+		return true;
+	}
+
+	IMPLEMENT_SIMPLE_AUTOMATION_TEST(FPS2VideoBidirectionalTest, "System.Plugins.PixelStreaming2.FPS2VideoBidirectionalTest", EAutomationTestFlags::EditorContext | EAutomationTestFlags::ClientContext | EAutomationTestFlags::ProductFilter)
+	bool FPS2VideoBidirectionalTest::RunTest(const FString& Parameters)
+	{
+		SetMediaDirection(EMediaType::Video, EMediaDirection::Bidirectional);
+		DoMediaDirectionTest(EMediaType::Video, EMediaDirection::Bidirectional);
+
+		return true;
+	}
+} // namespace UE::PixelStreaming2
+
+#endif // WITH_DEV_AUTOMATION_TESTS
\ No newline at end of file
diff --git a/Engine/Plugins/Media/PixelStreaming2/Source/PixelStreaming2RTC/Private/Tests/MockPlayer.cpp b/Engine/Plugins/Media/PixelStreaming2/Source/PixelStreaming2RTC/Private/Tests/MockPlayer.cpp
--- a/Engine/Plugins/Media/PixelStreaming2/Source/PixelStreaming2RTC/Private/Tests/MockPlayer.cpp
+++ b/Engine/Plugins/Media/PixelStreaming2/Source/PixelStreaming2RTC/Private/Tests/MockPlayer.cpp
@@ -30,7 +30,9 @@ namespace UE::PixelStreaming2
 		VideoBuffer.SafeRelease();
 	}
 
-	FMockPlayer::FMockPlayer()
+	FMockPlayer::FMockPlayer(FMockPlayerConfig Config)
+		: AudioDirection(Config.AudioDirection)
+		, VideoDirection(Config.VideoDirection)
 	{
 		PlayerName = FUtf8String::Printf("MockPlayer%d", PlayerId++);
 		FUtf8String ConferenceId = "test_conference";
@@ -54,6 +56,11 @@ namespace UE::PixelStreaming2
 		EpicRtcManager->OnVideoTrackMuted.AddRaw(this, &FMockPlayer::OnVideoTrackMuted);
 		EpicRtcManager->OnVideoTrackRemoved.AddRaw(this, &FMockPlayer::OnVideoTrackRemoved);
 		EpicRtcManager->OnVideoTrackState.AddRaw(this, &FMockPlayer::OnVideoTrackState);
+		EpicRtcManager->OnAudioTrackUpdate.AddRaw(this, &FMockPlayer::OnAudioTrackUpdate);
+		EpicRtcManager->OnAudioTrackFrame.AddRaw(this, &FMockPlayer::OnAudioTrackFrame);
+		EpicRtcManager->OnAudioTrackMuted.AddRaw(this, &FMockPlayer::OnAudioTrackMuted);
+		EpicRtcManager->OnAudioTrackRemoved.AddRaw(this, &FMockPlayer::OnAudioTrackRemoved);
+		EpicRtcManager->OnAudioTrackState.AddRaw(this, &FMockPlayer::OnAudioTrackState);
 
 		EpicRtcManager->OnDataTrackMessage.AddRaw(this, &FMockPlayer::OnDataTrackMessage);
 		EpicRtcManager->OnDataTrackRemoved.AddRaw(this, &FMockPlayer::OnDataTrackRemoved);
@@ -175,7 +182,7 @@ namespace UE::PixelStreaming2
 	{
 		FString ParticipantId{ (int32)Participant->GetId()._length, Participant->GetId()._ptr };
 		FString VideoTrackId{ (int32)VideoTrack->GetId()._length, VideoTrack->GetId()._ptr };
-		UE_LOG(LogPixelStreaming2RTC, VeryVerbose, TEXT("UPixelStreaming2Peer::OnVideoTrackUpdate(Participant [%s], VideoTrack [%s])"), *ParticipantId, *VideoTrackId);
+		UE_LOG(LogPixelStreaming2RTC, VeryVerbose, TEXT("FMockPlayer::OnVideoTrackUpdate(Participant [%s], VideoTrack [%s])"), *ParticipantId, *VideoTrackId);
 	}
 
 	void FMockPlayer::OnVideoTrackFrame(EpicRtcVideoTrackInterface* VideoTrack, const EpicRtcVideoFrame& Frame)
@@ -215,6 +222,53 @@ namespace UE::PixelStreaming2
 	{
 		FString ParticipantId{ (int32)Participant->GetId()._length, Participant->GetId()._ptr };
 		UE_LOG(LogPixelStreaming2RTC, Log, TEXT("OnRoomJoinedUpdate: Player (%s) joined"), *ParticipantId);
+
+		TRefCountPtr<EpicRtcConnectionInterface> ParticipantConnection = Participant->GetConnection();
+		ParticipantConnection->SetManualNegotiation(true);
+
+		const bool bSyncVideoAndAudio = !UPixelStreaming2PluginSettings::CVarWebRTCDisableAudioSync.GetValueOnAnyThread();
+
+		if (VideoDirection == EMediaDirection::SendOnly || VideoDirection == EMediaDirection::Bidirectional)
+		{
+			TArray<EpicRtcVideoEncodingConfig> VideoEncodingConfigs;
+
+			VideoEncodingConfigs.Add({
+				._rid = EpicRtcStringView{ ._ptr = nullptr, ._length = 0 },
+				._scaleResolutionDownBy = 1.0,
+				._scalabilityMode = EpicRtcVideoScalabilityMode::L1T1,
+				._minBitrate = 1'000'000,
+				._maxBitrate = 10'000'000,
+				._maxFrameRate = 60 //
+			});
+
+			EpicRtcVideoEncodingConfigSpan VideoEncodingConfigSpan = {
+				._ptr = VideoEncodingConfigs.GetData(),
+				._size = (uint64_t)VideoEncodingConfigs.Num()
+			};
+
+			FUtf8String		   VideoStreamID = bSyncVideoAndAudio ? "pixelstreaming_av_stream_id" : "pixelstreaming_video_stream_id";
+			EpicRtcVideoSource VideoSource = {
+				._streamId = ToEpicRtcStringView(VideoStreamID),
+				._encodings = VideoEncodingConfigSpan,
+				._direction = EpicRtcMediaSourceDirection::SendRecv
+			};
+
+			ParticipantConnection->AddVideoSource(VideoSource);
+		}
+
+		if (AudioDirection == EMediaDirection::SendOnly || AudioDirection == EMediaDirection::Bidirectional)
+		{
+
+			FUtf8String		   AudioStreamID = bSyncVideoAndAudio ? "pixelstreaming_av_stream_id" : "pixelstreaming_audio_stream_id";
+			EpicRtcAudioSource AudioSource = {
+				._streamId = ToEpicRtcStringView(AudioStreamID),
+				._bitrate = 510000,
+				._channels = 2,
+				._direction = EpicRtcMediaSourceDirection::SendRecv
+			};
+
+			ParticipantConnection->AddAudioSource(AudioSource);
+		}
 	}
 
 	void FMockPlayer::OnRoomLeftUpdate(const EpicRtcStringView ParticipantId)
@@ -274,6 +328,26 @@ namespace UE::PixelStreaming2
 		DataTrack = FEpicRtcDataTrack::Create(InDataTrack, ToStreamerProtocol);
 	}
 
+	void FMockPlayer::OnAudioTrackUpdate(EpicRtcParticipantInterface* Particpant, EpicRtcAudioTrackInterface* AudioTrack)
+	{
+	}
+
+	void FMockPlayer::OnAudioTrackMuted(EpicRtcAudioTrackInterface* AudioTrack, EpicRtcBool bIsMuted)
+	{
+	}
+
+	void FMockPlayer::OnAudioTrackFrame(EpicRtcAudioTrackInterface* AudioTrack, const EpicRtcAudioFrame& Frame)
+	{
+	}
+
+	void FMockPlayer::OnAudioTrackRemoved(EpicRtcAudioTrackInterface* AudioTrack)
+	{
+	}
+
+	void FMockPlayer::OnAudioTrackState(EpicRtcAudioTrackInterface* AudioTrack, const EpicRtcTrackState)
+	{
+	}
+
 	void FMockPlayer::Disconnect(const FString& Reason)
 	{
 		if (!EpicRtcManager->EpicRtcSession)
diff --git a/Engine/Plugins/Media/PixelStreaming2/Source/PixelStreaming2RTC/Private/Tests/MockPlayer.h b/Engine/Plugins/Media/PixelStreaming2/Source/PixelStreaming2RTC/Private/Tests/MockPlayer.h
--- a/Engine/Plugins/Media/PixelStreaming2/Source/PixelStreaming2RTC/Private/Tests/MockPlayer.h
+++ b/Engine/Plugins/Media/PixelStreaming2/Source/PixelStreaming2RTC/Private/Tests/MockPlayer.h
@@ -15,6 +15,14 @@ class EpicRtcPlatformInterface;
 
 namespace UE::PixelStreaming2
 {
+	enum class EMediaDirection : uint8
+	{
+		Disabled,
+		SendOnly,
+		RecvOnly,
+		Bidirectional
+	};
+
 	class FEpicRtcDataTrack;
 
 	struct FMockVideoFrameConfig
@@ -39,10 +47,16 @@ namespace UE::PixelStreaming2
 		FThreadSafeBool							  bReceivedFrame = false;
 	};
 
+	struct FMockPlayerConfig
+	{
+		EMediaDirection AudioDirection;
+		EMediaDirection VideoDirection;
+	};
+
 	class FMockPlayer : public TSharedFromThis<FMockPlayer>
 	{
 	public:
-		FMockPlayer();
+		FMockPlayer(FMockPlayerConfig Config = { .AudioDirection = EMediaDirection::RecvOnly, .VideoDirection = EMediaDirection::RecvOnly });
 		virtual ~FMockPlayer();
 
 		void Connect(int StreamerPort);
@@ -65,6 +79,11 @@ namespace UE::PixelStreaming2
 		void OnDataTrackRemoved(EpicRtcDataTrackInterface*);
 		void OnDataTrackState(EpicRtcDataTrackInterface*, const EpicRtcTrackState);
 		void OnDataTrackUpdate(EpicRtcParticipantInterface*, EpicRtcDataTrackInterface*);
+		void OnAudioTrackUpdate(EpicRtcParticipantInterface* Particpant, EpicRtcAudioTrackInterface* AudioTrack);
+		void OnAudioTrackMuted(EpicRtcAudioTrackInterface* AudioTrack, EpicRtcBool bIsMuted);
+		void OnAudioTrackFrame(EpicRtcAudioTrackInterface* AudioTrack, const EpicRtcAudioFrame& Frame);
+		void OnAudioTrackRemoved(EpicRtcAudioTrackInterface* AudioTrack);
+		void OnAudioTrackState(EpicRtcAudioTrackInterface* AudioTrack, const EpicRtcTrackState);
 
 		template <typename... Args>
 		bool SendMessage(FString MessageType, Args... VarArgs)
@@ -79,19 +98,20 @@ namespace UE::PixelStreaming2
 
 		bool DataChannelAvailable() { return DataTrack.IsValid(); };
 
-		TSharedPtr<FMockVideoSink>					   GetVideoSink() { return VideoSink; };
+		TSharedPtr<FMockVideoSink>				 GetVideoSink() { return VideoSink; };
 		TSharedPtr<IPixelStreaming2DataProtocol> GetToStreamerProtocol() { return ToStreamerProtocol; }
 
 		DECLARE_MULTICAST_DELEGATE_OneParam(FOnMessageReceived, const TArray<uint8>&);
 		FOnMessageReceived OnMessageReceived;
 
+		TSharedPtr<FEpicRtcManager> EpicRtcManager;
+
 	private:
-		TSharedPtr<FEpicRtcManager>					   EpicRtcManager;
-		TSharedPtr<FMockVideoSink>					   VideoSink;
-		TSharedPtr<FEpicRtcDataTrack>				   DataTrack;
-		TRefCountPtr<EpicRtcPlatformInterface>		   Platform;
-		TUniqueTaskPtr<FEpicRtcTickConferenceTask>	   TickConferenceTask;
-		TSharedPtr<IPixelStreaming2DataProtocol> ToStreamerProtocol;
+		TSharedPtr<FMockVideoSink>				   VideoSink;
+		TSharedPtr<FEpicRtcDataTrack>			   DataTrack;
+		TRefCountPtr<EpicRtcPlatformInterface>	   Platform;
+		TUniqueTaskPtr<FEpicRtcTickConferenceTask> TickConferenceTask;
+		TSharedPtr<IPixelStreaming2DataProtocol>   ToStreamerProtocol;
 
 		TArray<EpicRtcVideoEncoderInitializerInterface*> EpicRtcVideoEncoderInitializers;
 		TArray<EpicRtcVideoDecoderInitializerInterface*> EpicRtcVideoDecoderInitializers;
@@ -101,6 +121,9 @@ namespace UE::PixelStreaming2
 		FUtf8String		SubscribedStream;
 		FUtf8String		PlayerName;
 		static uint32_t PlayerId;
+
+		EMediaDirection AudioDirection;
+		EMediaDirection VideoDirection;
 	};
 } // namespace UE::PixelStreaming2
 
diff --git a/Engine/Plugins/Media/PixelStreaming2/Source/PixelStreaming2RTC/Private/Tests/TestUtils.cpp b/Engine/Plugins/Media/PixelStreaming2/Source/PixelStreaming2RTC/Private/Tests/TestUtils.cpp
--- a/Engine/Plugins/Media/PixelStreaming2/Source/PixelStreaming2RTC/Private/Tests/TestUtils.cpp
+++ b/Engine/Plugins/Media/PixelStreaming2/Source/PixelStreaming2RTC/Private/Tests/TestUtils.cpp
@@ -218,6 +218,31 @@ namespace UE::PixelStreaming2
 		return *bComplete.Get(); // Not connected or timed out so run this latent test again next frame
 	}
 
+	bool FWaitForPlayerTrackOrTimeout::Update()
+	{
+		double DeltaTime = FPlatformTime::Seconds() - StartTime;
+		if (DeltaTime > TimeoutSeconds)
+		{
+			UE_LOG(LogPixelStreaming2RTC, Error, TEXT("Player timed out waiting for a track."));
+			return true;
+		}
+		return *bComplete.Get();
+	}
+
+	bool FWaitAndCheckBool::Update()
+	{
+		double DeltaTime = FPlatformTime::Seconds() - StartTime;
+		if (DeltaTime > WaitSeconds)
+		{
+			if ((*bCheck.Get()) != bExpectedValue)
+			{
+				UE_LOGFMT(LogPixelStreaming2RTC, Error, "FWaitAndCheckBool failed. Expected [{0}] but got [{1}]", bExpectedValue, *bCheck.Get());
+			}
+			return true;
+		}
+		return false; // We keep updating this until the timeout is complete and then we check that there's no tracks added
+	}
+
 	bool FWaitForStreamerDataChannelMessageOrTimeout::Update()
 	{
 		double DeltaTime = FPlatformTime::Seconds() - StartTime;
@@ -263,6 +288,10 @@ namespace UE::PixelStreaming2
 			OutSignallingServer->Stop();
 			OutSignallingServer.Reset();
 		}
+
+		// Restore media directions back to default
+		SetMediaDirection(EMediaType::Audio, EMediaDirection::Bidirectional);
+		SetMediaDirection(EMediaType::Video, EMediaDirection::Bidirectional);
 		return true;
 	}
 
@@ -276,11 +305,21 @@ namespace UE::PixelStreaming2
 		});
 	}
 
-	void SetDisableTransmitVideo(bool bDisableTransmitVideo)
+	void SetMediaDirection(EMediaType MediaType, EMediaDirection Direction)
 	{
-		// Set whether to transmit video
-		UE::PixelStreaming2::DoOnGameThreadAndWait(MAX_uint32, [bDisableTransmitVideo]() {
-			UPixelStreaming2PluginSettings::CVarWebRTCDisableTransmitVideo->Set(bDisableTransmitVideo, ECVF_SetByCode);
+		bool bTransmit = Direction == EMediaDirection::SendOnly || Direction == EMediaDirection::Bidirectional;
+		bool bReceive = Direction == EMediaDirection::RecvOnly || Direction == EMediaDirection::Bidirectional;
+		DoOnGameThreadAndWait(MAX_uint32, [MediaType, bTransmit, bReceive]() {
+			if (MediaType == EMediaType::Audio)
+			{
+				UPixelStreaming2PluginSettings::CVarWebRTCDisableTransmitAudio->SetWithCurrentPriority(!bTransmit);
+				UPixelStreaming2PluginSettings::CVarWebRTCDisableReceiveAudio->SetWithCurrentPriority(!bReceive);
+			}
+			else if (MediaType == EMediaType::Video)
+			{
+				UPixelStreaming2PluginSettings::CVarWebRTCDisableTransmitVideo->SetWithCurrentPriority(!bTransmit);
+				UPixelStreaming2PluginSettings::CVarWebRTCDisableReceiveVideo->SetWithCurrentPriority(!bReceive);
+			}
 		});
 	}
 
@@ -297,9 +336,9 @@ namespace UE::PixelStreaming2
 		return OutStreamer;
 	}
 
-	TSharedPtr<FMockPlayer> CreatePlayer()
+	TSharedPtr<FMockPlayer> CreatePlayer(FMockPlayerConfig Config)
 	{
-		TSharedPtr<FMockPlayer> OutPlayer = MakeShared<FMockPlayer>();
+		TSharedPtr<FMockPlayer> OutPlayer = MakeShared<FMockPlayer>(Config);
 
 		return OutPlayer;
 	}
diff --git a/Engine/Plugins/Media/PixelStreaming2/Source/PixelStreaming2RTC/Private/Tests/TestUtils.h b/Engine/Plugins/Media/PixelStreaming2/Source/PixelStreaming2RTC/Private/Tests/TestUtils.h
--- a/Engine/Plugins/Media/PixelStreaming2/Source/PixelStreaming2RTC/Private/Tests/TestUtils.h
+++ b/Engine/Plugins/Media/PixelStreaming2/Source/PixelStreaming2RTC/Private/Tests/TestUtils.h
@@ -13,6 +13,12 @@
 
 namespace UE::PixelStreaming2
 {
+	enum class EMediaType : uint8
+	{
+		Audio,
+		Video
+	};
+
 	namespace TestUtils
 	{
 		int32 NextStreamerPort();
@@ -64,9 +70,9 @@ namespace UE::PixelStreaming2
 		virtual bool Update() override;
 
 	private:
-		double									   TimeoutSeconds;
+		double								 TimeoutSeconds;
 		TSharedPtr<IPixelStreaming2Streamer> Streamer;
-		TSharedPtr<bool>						   bComplete;
+		TSharedPtr<bool>					 bComplete;
 	};
 
 	DEFINE_LATENT_AUTOMATION_COMMAND_ONE_PARAMETER(FWaitSeconds, double, WaitSeconds);
@@ -78,12 +84,14 @@ namespace UE::PixelStreaming2
 	DEFINE_LATENT_AUTOMATION_COMMAND_THREE_PARAMETER(FSendDataChannelMessageToStreamer, TSharedPtr<FMockPlayer>, Player, FString, MessageType, const FString, Body);
 	DEFINE_LATENT_AUTOMATION_COMMAND_THREE_PARAMETER(FSendDataChannelMessageFromStreamer, TSharedPtr<IPixelStreaming2Streamer>, Streamer, FString, MessageType, const FString, Body);
 	DEFINE_LATENT_AUTOMATION_COMMAND_THREE_PARAMETER(FSendCustomMessageToStreamer, TSharedPtr<FMockPlayer>, Player, FString, MessageType, uint16, Body);
+	DEFINE_LATENT_AUTOMATION_COMMAND_THREE_PARAMETER(FWaitForPlayerTrackOrTimeout, double, TimeoutSeconds, TSharedPtr<FMockPlayer>, Player, TSharedPtr<bool>, bComplete);
+	DEFINE_LATENT_AUTOMATION_COMMAND_FOUR_PARAMETER(FWaitAndCheckBool, double, WaitSeconds, TSharedPtr<FMockPlayer>, Player, TSharedPtr<bool>, bCheck, bool, bExpectedValue);
 
-	TSharedPtr<IPixelStreaming2Streamer>			  CreateStreamer(const FString& StreamerName, int StreamerPort);
-	TSharedPtr<FMockPlayer>								  CreatePlayer();
+	TSharedPtr<IPixelStreaming2Streamer>			CreateStreamer(const FString& StreamerName, int StreamerPort);
+	TSharedPtr<FMockPlayer>							CreatePlayer(FMockPlayerConfig Config = { .AudioDirection = EMediaDirection::RecvOnly, .VideoDirection = EMediaDirection::RecvOnly });
 	TSharedPtr<UE::PixelStreaming2Servers::IServer> CreateSignallingServer(int StreamerPort, int PlayerPort);
-	void												  SetCodec(EVideoCodec Codec);
-	void												  SetDisableTransmitVideo(bool bDisableTransmitVideo);
+	void											SetCodec(EVideoCodec Codec);
+	void											SetMediaDirection(EMediaType Media, EMediaDirection Direction);
 } // namespace UE::PixelStreaming2
 
 #endif // WITH_DEV_AUTOMATION_TESTS
\ No newline at end of file
