[PS2] Add: Property and CVar to select the type of default streamer to be created

diff --git a/Engine/Plugins/Media/PixelStreaming2/Source/PixelStreaming2/Private/PixelStreaming2Module.cpp b/Engine/Plugins/Media/PixelStreaming2/Source/PixelStreaming2/Private/PixelStreaming2Module.cpp
--- a/Engine/Plugins/Media/PixelStreaming2/Source/PixelStreaming2/Private/PixelStreaming2Module.cpp
+++ b/Engine/Plugins/Media/PixelStreaming2/Source/PixelStreaming2/Private/PixelStreaming2Module.cpp
@@ -55,6 +55,26 @@ namespace UE::PixelStreaming2
 		// Initialize PixelStreaming thread. Handles tasks like audio pushing and conference ticking
 		PixelStreamingThread = MakeShared<FPixelStreamingThread>();
 
+		// Extra Initialisations post loading console commands
+		IConsoleManager::Get().RegisterConsoleCommand(
+			TEXT("PixelStreaming.StartStreaming"),
+			TEXT("Start all streaming sessions"),
+			FConsoleCommandDelegate::CreateLambda([DefaultStreamer = this->DefaultStreamer]() {
+				if (DefaultStreamer)
+				{
+					DefaultStreamer->StartStreaming();
+				}
+			}));
+		IConsoleManager::Get().RegisterConsoleCommand(
+			TEXT("PixelStreaming.StopStreaming"),
+			TEXT("End any existing streaming sessions."),
+			FConsoleCommandDelegate::CreateLambda([DefaultStreamer = this->DefaultStreamer]() {
+				if (DefaultStreamer)
+				{
+					DefaultStreamer->StopStreaming();
+				}
+			}));
+
 		FCoreDelegates::OnAllModuleLoadingPhasesComplete.AddLambda([this]() {
 			if (!IsPlatformSupported())
 			{
@@ -97,6 +117,20 @@ namespace UE::PixelStreaming2
 
 			bModuleReady = true;
 			ReadyEvent.Broadcast(*this);
+
+			// This is called after the ready event is broadcast so that other modules have a chance to initialize
+			// themselves before streamer creation
+			if (!GIsEditor)
+			{
+				// We don't want to create the default streamer when using the editor
+				InitDefaultStreamer();
+
+				// The user has specified a URL on the command line or ini meaning their intention is to start streaming immediately
+				if (!DefaultStreamer->GetConnectionURL().IsEmpty())
+				{
+					DefaultStreamer->StartStreaming();
+				}
+			}
 		});
 	}
 
@@ -110,6 +144,8 @@ namespace UE::PixelStreaming2
 		// We explicitly call release on streamer so WebRTC gets shutdown before our module is deleted
 		// additionally the streamer does a bunch of delegate calls and unbinds which seem to have issues
 		// when called during engine destruction rather than here.
+		DefaultStreamer.Reset();
+
 		Streamers.Empty();
 
 		// Reset thread must be called before tasks to ensure it does not attempt to run any partially destroyed tasks from the AudioMixingCapturer
@@ -358,6 +394,11 @@ namespace UE::PixelStreaming2
 		}
 	}
 
+	void FPixelStreaming2Module::InitDefaultStreamer()
+	{
+		UE_LOGFMT(LogPixelStreaming2, Log, "Initializing default streamer. ID: [{0}], Type: [{1}]", GetDefaultStreamerID(), UPixelStreaming2PluginSettings::CVarDefaultStreamerType.GetValueOnAnyThread());
+		DefaultStreamer = CreateStreamer(GetDefaultStreamerID(), UPixelStreaming2PluginSettings::CVarDefaultStreamerType.GetValueOnAnyThread());
+	}
 } // namespace UE::PixelStreaming2
 
 IMPLEMENT_MODULE(UE::PixelStreaming2::FPixelStreaming2Module, PixelStreaming2)
diff --git a/Engine/Plugins/Media/PixelStreaming2/Source/PixelStreaming2/Private/PixelStreaming2Module.h b/Engine/Plugins/Media/PixelStreaming2/Source/PixelStreaming2/Private/PixelStreaming2Module.h
--- a/Engine/Plugins/Media/PixelStreaming2/Source/PixelStreaming2/Private/PixelStreaming2Module.h
+++ b/Engine/Plugins/Media/PixelStreaming2/Source/PixelStreaming2/Private/PixelStreaming2Module.h
@@ -49,6 +49,8 @@ namespace UE::PixelStreaming2
 		void RegisterCustomHandlers(TSharedPtr<IPixelStreaming2Streamer> Streamer);
 		void HandleUIInteraction(FMemoryReader Ar);
 
+		void InitDefaultStreamer();
+
 	private:
 		bool		bModuleReady = false;
 		FReadyEvent ReadyEvent;
@@ -58,6 +60,8 @@ namespace UE::PixelStreaming2
 		// FPixelStreamingThread must exist before any AudioTask and AudioMixingCapturer(Which contains a audio task) to ensure it is destroyed last
 		TSharedPtr<class FPixelStreamingThread> PixelStreamingThread;
 
+		TSharedPtr<IPixelStreaming2Streamer> DefaultStreamer;
+
 		TThreadSafeMap<FString, TWeakPtr<IPixelStreaming2Streamer>> Streamers;
 
 		static FPixelStreaming2Module* PixelStreaming2Module;
diff --git a/Engine/Plugins/Media/PixelStreaming2/Source/PixelStreaming2RTC/Private/EpicRtcStreamer.cpp b/Engine/Plugins/Media/PixelStreaming2/Source/PixelStreaming2RTC/Private/EpicRtcStreamer.cpp
--- a/Engine/Plugins/Media/PixelStreaming2/Source/PixelStreaming2RTC/Private/EpicRtcStreamer.cpp
+++ b/Engine/Plugins/Media/PixelStreaming2/Source/PixelStreaming2RTC/Private/EpicRtcStreamer.cpp
@@ -4,6 +4,8 @@
 
 #include "PixelStreaming2PluginSettings.h"
 #include "PixelStreaming2Delegates.h"
+#include "Engine/GameEngine.h"
+#include "Engine/GameViewportClient.h"
 #include "EpicRtcWebsocket.h"
 #include "EpicRtcWebsocketFactory.h"
 #include "TextureResource.h"
@@ -26,6 +28,8 @@
 #include "CoreGlobals.h"
 #include "UtilsString.h"
 #include "UtilsCodecs.h"
+#include "VideoProducerBackBuffer.h"
+#include "VideoProducerMediaCapture.h"
 
 namespace UE::PixelStreaming2
 {
@@ -1635,6 +1639,37 @@ namespace UE::PixelStreaming2
 
 	TSharedPtr<IPixelStreaming2Streamer> FRTCStreamerFactory::CreateNewStreamer(const FString& StreamerId)
 	{
-		return MakeShared<FEpicRtcStreamer>(StreamerId, Conference);
+		TSharedPtr<IPixelStreaming2Streamer> NewStreamer = MakeShared<FEpicRtcStreamer>(StreamerId, Conference);
+
+		// default to the scene viewport if we have a game engine
+		if (UGameEngine* GameEngine = Cast<UGameEngine>(GEngine))
+		{
+			TSharedPtr<SWindow>						 TargetWindow = GameEngine->GameViewport->GetWindow();
+			TSharedPtr<IPixelStreaming2InputHandler> InputHandler = NewStreamer->GetInputHandler().Pin();
+			if (TargetWindow.IsValid() && InputHandler.IsValid())
+			{
+				InputHandler->SetTargetWindow(TargetWindow);
+			}
+			else
+			{
+				UE_LOG(LogPixelStreaming2RTC, Error, TEXT("Cannot set target window - target window is not valid."));
+			}
+		}
+
+		// If the user hasn't specified a connection url on the command line or in the ini, don't set
+		// the video producer in order to not tax their GPU unnecessarily
+		if (!UPixelStreaming2PluginSettings::CVarConnectionURL.GetValueOnAnyThread().IsEmpty())
+		{
+			if (UPixelStreaming2PluginSettings::CVarUseMediaCapture.GetValueOnAnyThread())
+			{
+				NewStreamer->SetVideoProducer(FVideoProducerMediaCapture::CreateActiveViewportCapture());
+			}
+			else
+			{
+				NewStreamer->SetVideoProducer(FVideoProducerBackBuffer::Create());
+			}
+		}
+
+		return NewStreamer;
 	}
 } // namespace UE::PixelStreaming2
diff --git a/Engine/Plugins/Media/PixelStreaming2/Source/PixelStreaming2RTC/Private/PixelStreaming2RTCModule.cpp b/Engine/Plugins/Media/PixelStreaming2/Source/PixelStreaming2RTC/Private/PixelStreaming2RTCModule.cpp
--- a/Engine/Plugins/Media/PixelStreaming2/Source/PixelStreaming2RTC/Private/PixelStreaming2RTCModule.cpp
+++ b/Engine/Plugins/Media/PixelStreaming2/Source/PixelStreaming2RTC/Private/PixelStreaming2RTCModule.cpp
@@ -150,16 +150,6 @@ namespace UE::PixelStreaming2
 			// Ensure we have ImageWrapper loaded, used in Freezeframes
 			verify(FModuleManager::Get().LoadModule(FName("ImageWrapper")));
 
-			// We don't want to start immediately streaming in editor
-			if (!GIsEditor)
-			{
-				InitDefaultStreamer();
-				if (DefaultStreamer)
-				{
-					DefaultStreamer->StartStreaming();
-				}
-			}
-
 			bModuleReady = true;
 			ReadyEvent.Broadcast(*this);
 		});
@@ -169,27 +159,6 @@ namespace UE::PixelStreaming2
 		// Call these to initialize their singletons
 		FStats::Get();
 
-		// Extra Initialisations post loading console commands
-		IConsoleManager::Get().RegisterConsoleCommand(
-			TEXT("PixelStreaming.StartStreaming"),
-			TEXT("Start all streaming sessions"),
-			FConsoleCommandDelegate::CreateLambda([DefaultStreamer = this->DefaultStreamer]() {
-				if (DefaultStreamer)
-				{
-					DefaultStreamer->StartStreaming();
-				}
-			}));
-
-		IConsoleManager::Get().RegisterConsoleCommand(
-			TEXT("PixelStreaming.StopStreaming"),
-			TEXT("End any existing streaming sessions."),
-			FConsoleCommandDelegate::CreateLambda([DefaultStreamer = this->DefaultStreamer]() {
-				if (DefaultStreamer)
-				{
-					DefaultStreamer->StopStreaming();
-				}
-			}));
-
 		if (UPixelStreaming2PluginSettings::FDelegates* Delegates = UPixelStreaming2PluginSettings::Delegates())
 		{
 			Delegates->OnLogStatsChanged.AddLambda([this](IConsoleVariable* Var) {
@@ -237,11 +206,6 @@ namespace UE::PixelStreaming2
 			return;
 		}
 
-		// We explicitly call release on streamer so WebRTC gets shutdown before our module is deleted
-		// additionally the streamer does a bunch of delegate calls and unbinds which seem to have issues
-		// when called during engine destruction rather than here.
-		DefaultStreamer.Reset();
-
 		AudioMixingCapturer.Reset();
 
 		TickConferenceTask.Reset();
@@ -299,51 +263,6 @@ namespace UE::PixelStreaming2
 		return AudioMixingCapturer;
 	}
 
-	void FPixelStreaming2RTCModule::InitDefaultStreamer()
-	{
-		IPixelStreaming2Module& CoreModule = IPixelStreaming2Module::Get();
-
-		FString ConnectionURL = CoreModule.GetDefaultConnectionURL();
-		FURL	URL(nullptr, *ConnectionURL, TRAVEL_Absolute);
-		if (URL.Protocol != TEXT("ws") && URL.Protocol != TEXT("wss"))
-		{
-			UE_LOGFMT(LogPixelStreaming2RTC, Log, "PixelStreaming2RTC early out due to protocol not matching ws or wss");
-		}
-
-		UE_LOGFMT(LogPixelStreaming2RTC, Log, "PixelStreaming2 streamer ID: {0}", CoreModule.GetDefaultStreamerID());
-
-		DefaultStreamer = CoreModule.CreateStreamer(CoreModule.GetDefaultStreamerID(), StreamerFactory->GetStreamType());
-
-		// default to the scene viewport if we have a game engine
-		if (UGameEngine* GameEngine = Cast<UGameEngine>(GEngine))
-		{
-			TSharedPtr<SWindow>						 TargetWindow = GameEngine->GameViewport->GetWindow();
-			TSharedPtr<IPixelStreaming2InputHandler> InputHandler = DefaultStreamer->GetInputHandler().Pin();
-			if (TargetWindow.IsValid() && InputHandler.IsValid())
-			{
-				InputHandler->SetTargetWindow(TargetWindow);
-			}
-			else
-			{
-				UE_LOG(LogPixelStreaming2RTC, Error, TEXT("Cannot set target window - target window is not valid."));
-			}
-		}
-
-		if (!DefaultStreamer->GetConnectionURL().IsEmpty())
-		{
-			// The user has specified a URL on the command line meaning their intention is to start streaming immediately
-			// in that case, set up the video input for them (as long as we're not in editor)
-			if (UPixelStreaming2PluginSettings::CVarUseMediaCapture.GetValueOnAnyThread())
-			{
-				DefaultStreamer->SetVideoProducer(FVideoProducerMediaCapture::CreateActiveViewportCapture());
-			}
-			else
-			{
-				DefaultStreamer->SetVideoProducer(FVideoProducerBackBuffer::Create());
-			}
-		}
-	}
-
 	FString FPixelStreaming2RTCModule::GetFieldTrials()
 	{
 		FString FieldTrials = UPixelStreaming2PluginSettings::CVarWebRTCFieldTrials.GetValueOnAnyThread();
diff --git a/Engine/Plugins/Media/PixelStreaming2/Source/PixelStreaming2RTC/Private/PixelStreaming2RTCModule.h b/Engine/Plugins/Media/PixelStreaming2/Source/PixelStreaming2RTC/Private/PixelStreaming2RTCModule.h
--- a/Engine/Plugins/Media/PixelStreaming2/Source/PixelStreaming2RTC/Private/PixelStreaming2RTCModule.h
+++ b/Engine/Plugins/Media/PixelStreaming2/Source/PixelStreaming2RTC/Private/PixelStreaming2RTCModule.h
@@ -44,9 +44,6 @@ namespace UE::PixelStreaming2
 		void ShutdownModule() override;
 		// End IModuleInterface
 
-		// Own methods
-		void InitDefaultStreamer();
-
 		FString GetFieldTrials();
 		bool	InitializeEpicRtc();
 
@@ -58,8 +55,6 @@ namespace UE::PixelStreaming2
 		FReadyEvent		ReadyEvent;
 		FDelegateHandle LogStatsHandle;
 
-		TSharedPtr<IPixelStreaming2Streamer> DefaultStreamer;
-
 	private:
 		TSharedPtr<class FEpicRtcAudioCapturer>	 AudioMixingCapturer;
 		TRefCountPtr<EpicRtcPlatformInterface>	 EpicRtcPlatform;
diff --git a/Engine/Plugins/Media/PixelStreaming2/Source/PixelStreaming2Settings/Internal/PixelStreaming2PluginSettings.h b/Engine/Plugins/Media/PixelStreaming2/Source/PixelStreaming2Settings/Internal/PixelStreaming2PluginSettings.h
--- a/Engine/Plugins/Media/PixelStreaming2/Source/PixelStreaming2Settings/Internal/PixelStreaming2PluginSettings.h
+++ b/Engine/Plugins/Media/PixelStreaming2/Source/PixelStreaming2Settings/Internal/PixelStreaming2PluginSettings.h
@@ -137,6 +137,13 @@ public:
 		))
 	FString DefaultStreamerID = TEXT("DefaultStreamer");
 
+	static TAutoConsoleVariable<FString> CVarDefaultStreamerType;
+	UPROPERTY(config, EditAnywhere, Category = "PixelStreaming", meta = (
+		DisplayName = "Default Streamer Type",
+		ToolTip = "Default Streamer Type to be used when not specified elsewhere. This value should match a type registered by an IPixelStreaming2StreamerFactory"
+		))
+	FString DefaultStreamerType = TEXT("DefaultRtc");
+
 	static TAutoConsoleVariable<FString> CVarConnectionURL;
 	UPROPERTY(config, EditAnywhere, Category = "PixelStreaming", meta = (
 		DisplayName = "Default Connection URL",
diff --git a/Engine/Plugins/Media/PixelStreaming2/Source/PixelStreaming2Settings/Private/PixelStreaming2PluginSettings.cpp b/Engine/Plugins/Media/PixelStreaming2/Source/PixelStreaming2Settings/Private/PixelStreaming2PluginSettings.cpp
--- a/Engine/Plugins/Media/PixelStreaming2/Source/PixelStreaming2Settings/Private/PixelStreaming2PluginSettings.cpp
+++ b/Engine/Plugins/Media/PixelStreaming2/Source/PixelStreaming2Settings/Private/PixelStreaming2PluginSettings.cpp
@@ -1,6 +1,7 @@
 // Copyright Epic Games, Inc. All Rights Reserved.
 
 #include "PixelStreaming2PluginSettings.h"
+
 #include "Logging.h"
 #include "Misc/CommandLine.h"
 #include "UObject/ReflectedTypeAccessors.h"
@@ -140,6 +141,7 @@ static const TSet<TPair<FString, FString>> GetCmdArg = {
 	{ "PixelStreaming2.SignalingKeepAliveInterval", "SignalingKeepAliveInterval" },
 	{ "PixelStreaming2.UseMediaCapture", "UseMediaCapture" },
 	{ "PixelStreaming2.ID", "DefaultStreamerID" },
+	{ "PixelStreaming2.DefaultStreamerType", "DefaultStreamerType" },
 	{ "PixelStreaming2.ConnectionURL", "ConnectionURL" },
 	{ "PixelStreaming2.CaptureUseFence", "CaptureUseFence" },
 	{ "PixelStreaming2.Encoder.Codec", "Codec" },
@@ -268,6 +270,12 @@ TAutoConsoleVariable<FString> UPixelStreaming2PluginSettings::CVarDefaultStreame
 	TEXT("Default Streamer ID to be used when not specified elsewhere."),
 	ECVF_Default);
 
+TAutoConsoleVariable<FString> UPixelStreaming2PluginSettings::CVarDefaultStreamerType(
+	TEXT("PixelStreaming2.DefaultStreamerType"),
+	TEXT("DefaultRtc"),
+	TEXT("Default Streamer Type to be used when not specified elsewhere."),
+	ECVF_Default);
+
 TAutoConsoleVariable<FString> UPixelStreaming2PluginSettings::CVarConnectionURL(
 	TEXT("PixelStreaming2.ConnectionURL"),
 	TEXT(""),
