[PS2] QOL: Refactor module structure so settings can show available streamer factories

diff --git a/Engine/Plugins/Media/PixelStreaming2/PixelStreaming2.uplugin b/Engine/Plugins/Media/PixelStreaming2/PixelStreaming2.uplugin
--- a/Engine/Plugins/Media/PixelStreaming2/PixelStreaming2.uplugin
+++ b/Engine/Plugins/Media/PixelStreaming2/PixelStreaming2.uplugin
@@ -15,6 +15,13 @@
 	"IsBetaVersion": false,
 	"Installed": false,
 	"Modules": [
+		{
+			"Name": "PixelStreaming2Core",
+			"Type": "Runtime",
+			"LoadingPhase": "PreDefault",
+			"PlatformAllowList": [ "Win64", "Linux", "Mac" ],
+			"TargetDenyList": [ ]
+		},
 		{
 			"Name": "PixelStreaming2",
 			"Type": "Runtime",
diff --git a/Engine/Plugins/Media/PixelStreaming2/Source/PixelStreaming2/PixelStreaming2.Build.cs b/Engine/Plugins/Media/PixelStreaming2/Source/PixelStreaming2/PixelStreaming2.Build.cs
--- a/Engine/Plugins/Media/PixelStreaming2/Source/PixelStreaming2/PixelStreaming2.Build.cs
+++ b/Engine/Plugins/Media/PixelStreaming2/Source/PixelStreaming2/PixelStreaming2.Build.cs
@@ -22,6 +22,7 @@ namespace UnrealBuildTool.Rules
                 "MediaUtils",
                 "PixelCapture",
                 "PixelStreaming2Input",
+                "PixelStreaming2Core",
                 "RHI",
                 "Renderer",
                 "RenderCore",
diff --git a/Engine/Plugins/Media/PixelStreaming2/Source/PixelStreaming2/Private/Streamer.cpp b/Engine/Plugins/Media/PixelStreaming2/Source/PixelStreaming2/Private/Streamer.cpp
deleted file mode 100644
diff --git a/Engine/Plugins/Media/PixelStreaming2/Source/PixelStreaming2/Private/Streamer.h b/Engine/Plugins/Media/PixelStreaming2/Source/PixelStreaming2/Private/Streamer.h
deleted file mode 100644
diff --git a/Engine/Plugins/Media/PixelStreaming2/Source/PixelStreaming2/Public/IPixelStreaming2AudioConsumer.h b/Engine/Plugins/Media/PixelStreaming2/Source/PixelStreaming2/Public/IPixelStreaming2AudioConsumer.h
deleted file mode 100644
diff --git a/Engine/Plugins/Media/PixelStreaming2/Source/PixelStreaming2/Public/IPixelStreaming2AudioProducer.h b/Engine/Plugins/Media/PixelStreaming2/Source/PixelStreaming2/Public/IPixelStreaming2AudioProducer.h
deleted file mode 100644
diff --git a/Engine/Plugins/Media/PixelStreaming2/Source/PixelStreaming2/Public/IPixelStreaming2AudioSink.h b/Engine/Plugins/Media/PixelStreaming2/Source/PixelStreaming2/Public/IPixelStreaming2AudioSink.h
deleted file mode 100644
diff --git a/Engine/Plugins/Media/PixelStreaming2/Source/PixelStreaming2/Public/IPixelStreaming2Module.h b/Engine/Plugins/Media/PixelStreaming2/Source/PixelStreaming2/Public/IPixelStreaming2Module.h
--- a/Engine/Plugins/Media/PixelStreaming2/Source/PixelStreaming2/Public/IPixelStreaming2Module.h
+++ b/Engine/Plugins/Media/PixelStreaming2/Source/PixelStreaming2/Public/IPixelStreaming2Module.h
@@ -8,6 +8,7 @@
 #include "Templates/RefCounting.h"
 #include "Templates/SharedPointer.h"
 #include "IPixelStreaming2AudioProducer.h"
+#include "IPixelStreaming2InputHandler.h"
 #include "IPixelStreaming2Streamer.h"
 #include "IPixelStreaming2VideoProducer.h"
 
diff --git a/Engine/Plugins/Media/PixelStreaming2/Source/PixelStreaming2/Public/IPixelStreaming2Streamer.h b/Engine/Plugins/Media/PixelStreaming2/Source/PixelStreaming2/Public/IPixelStreaming2Streamer.h
deleted file mode 100644
diff --git a/Engine/Plugins/Media/PixelStreaming2/Source/PixelStreaming2/Public/IPixelStreaming2VideoConsumer.h b/Engine/Plugins/Media/PixelStreaming2/Source/PixelStreaming2/Public/IPixelStreaming2VideoConsumer.h
deleted file mode 100644
diff --git a/Engine/Plugins/Media/PixelStreaming2/Source/PixelStreaming2/Public/IPixelStreaming2VideoProducer.h b/Engine/Plugins/Media/PixelStreaming2/Source/PixelStreaming2/Public/IPixelStreaming2VideoProducer.h
deleted file mode 100644
diff --git a/Engine/Plugins/Media/PixelStreaming2/Source/PixelStreaming2/Public/IPixelStreaming2VideoSink.h b/Engine/Plugins/Media/PixelStreaming2/Source/PixelStreaming2/Public/IPixelStreaming2VideoSink.h
deleted file mode 100644
diff --git a/Engine/Plugins/Media/PixelStreaming2/Source/PixelStreaming2Core/PixelStreaming2Core.Build.cs b/Engine/Plugins/Media/PixelStreaming2/Source/PixelStreaming2Core/PixelStreaming2Core.Build.cs
new file mode 100644
--- /dev/null
+++ b/Engine/Plugins/Media/PixelStreaming2/Source/PixelStreaming2Core/PixelStreaming2Core.Build.cs
@@ -0,0 +1,25 @@
+// Copyright Epic Games, Inc. All Rights Reserved.
+
+using System.Collections.Generic;
+using System.IO;
+
+namespace UnrealBuildTool.Rules
+{
+    public class PixelStreaming2Core : ModuleRules
+    {
+        public PixelStreaming2Core(ReadOnlyTargetRules Target) : base(Target)
+        {
+            PublicDependencyModuleNames.AddRange(new string[] { 
+                "CoreUObject",
+                "PixelCapture",
+                "RHI",
+                "WebRTC"
+            });
+
+            PrivateDependencyModuleNames.AddRange(new string[]
+            {
+                "Core"
+            });
+        }
+    }
+}
diff --git a/Engine/Plugins/Media/PixelStreaming2/Source/PixelStreaming2Core/Private/Logging.cpp b/Engine/Plugins/Media/PixelStreaming2/Source/PixelStreaming2Core/Private/Logging.cpp
new file mode 100644
--- /dev/null
+++ b/Engine/Plugins/Media/PixelStreaming2/Source/PixelStreaming2Core/Private/Logging.cpp
@@ -0,0 +1,5 @@
+// Copyright Epic Games, Inc. All Rights Reserved.
+
+#include "Logging.h"
+
+DEFINE_LOG_CATEGORY(LogPixelStreaming2Core);
\ No newline at end of file
diff --git a/Engine/Plugins/Media/PixelStreaming2/Source/PixelStreaming2Core/Private/Logging.h b/Engine/Plugins/Media/PixelStreaming2/Source/PixelStreaming2Core/Private/Logging.h
new file mode 100644
--- /dev/null
+++ b/Engine/Plugins/Media/PixelStreaming2/Source/PixelStreaming2Core/Private/Logging.h
@@ -0,0 +1,9 @@
+// Copyright Epic Games, Inc. All Rights Reserved.
+
+#pragma once
+
+#include "HAL/IConsoleManager.h"
+#include "Logging/LogMacros.h"
+#include "Logging/StructuredLog.h"
+
+DECLARE_LOG_CATEGORY_EXTERN(LogPixelStreaming2Core, Log, All);
\ No newline at end of file
diff --git a/Engine/Plugins/Media/PixelStreaming2/Source/PixelStreaming2Core/Private/PixelStreaming2CoreModule.cpp b/Engine/Plugins/Media/PixelStreaming2/Source/PixelStreaming2Core/Private/PixelStreaming2CoreModule.cpp
new file mode 100644
--- /dev/null
+++ b/Engine/Plugins/Media/PixelStreaming2/Source/PixelStreaming2Core/Private/PixelStreaming2CoreModule.cpp
@@ -0,0 +1,15 @@
+// Copyright Epic Games, Inc. All Rights Reserved.
+
+#include "Modules/ModuleInterface.h"
+#include "Modules/ModuleManager.h"
+
+namespace UE::PixelStreaming2
+{
+	class FPixelStreaming2CoreModule : public IModuleInterface
+	{
+	public:
+		virtual ~FPixelStreaming2CoreModule() = default;
+	};
+} // namespace UE::PixelStreaming2
+
+IMPLEMENT_MODULE(UE::PixelStreaming2::FPixelStreaming2CoreModule, PixelStreaming2Core)
\ No newline at end of file
diff --git a/Engine/Plugins/Media/PixelStreaming2/Source/PixelStreaming2Core/Private/Streamer.cpp b/Engine/Plugins/Media/PixelStreaming2/Source/PixelStreaming2Core/Private/Streamer.cpp
new file mode 100644
--- /dev/null
+++ b/Engine/Plugins/Media/PixelStreaming2/Source/PixelStreaming2Core/Private/Streamer.cpp
@@ -0,0 +1,127 @@
+// Copyright Epic Games, Inc. All Rights Reserved.
+
+#include "Streamer.h"
+
+#include "Features/IModularFeatures.h"
+#include "Logging.h"
+
+static FName ModularFeatureName = TEXT("PixelStreaming2 Streamer");
+
+namespace UE::PixelStreaming2
+{
+	class FDummyStreamer : public IPixelStreaming2Streamer
+	{
+	public:
+		FDummyStreamer() = default;
+		virtual ~FDummyStreamer() = default;
+
+		virtual void										 Initialize() override {}
+		virtual void										 SetStreamFPS(int32 InFramesPerSecond) override {}
+		virtual int32										 GetStreamFPS() override { return 0; }
+		virtual void										 SetCoupleFramerate(bool bCouple) override {}
+		virtual void										 SetVideoProducer(TSharedPtr<IPixelStreaming2VideoProducer> Input) override {}
+		virtual TWeakPtr<IPixelStreaming2VideoProducer>		 GetVideoProducer() override { return nullptr; }
+		virtual void										 SetConnectionURL(const FString& InSignallingServerURL) override {}
+		virtual FString										 GetConnectionURL() override { return FString(); }
+		virtual FString										 GetId() override { return TEXT("DummyStreamer"); }
+		virtual bool										 IsConnected() override { return false; }
+		virtual void										 StartStreaming() override {}
+		virtual void										 StopStreaming() override {}
+		virtual bool										 IsStreaming() const { return false; }
+		virtual FPreConnectionEvent&						 OnPreConnection() override { return StreamingPreConnectionEvent; }
+		virtual FStreamingStartedEvent&						 OnStreamingStarted() override { return StreamingStartedEvent; }
+		virtual FStreamingStoppedEvent&						 OnStreamingStopped() override { return StreamingStoppedEvent; }
+		virtual void										 ForceKeyFrame() override {}
+		virtual void										 FreezeStream(UTexture2D* Texture) override {}
+		virtual void										 UnfreezeStream() override {}
+		virtual void										 SendAllPlayersMessage(FString MessageType, const FString& Descriptor) override {}
+		virtual void										 SendPlayerMessage(FString PlayerId, FString MessageType, const FString& Descriptor) override {}
+		virtual void										 SendFileData(const TArray64<uint8>& ByteData, FString& MimeType, FString& FileExtension) override {}
+		virtual void										 KickPlayer(FString PlayerId) override {}
+		virtual TArray<FString>								 GetConnectedPlayers() override { return {}; }
+		virtual TWeakPtr<class IPixelStreaming2InputHandler> GetInputHandler() override { return nullptr; }
+		virtual IPixelStreaming2AudioSink*					 GetPeerAudioSink(FString PlayerId) override { return nullptr; }
+		virtual IPixelStreaming2AudioSink*					 GetUnlistenedAudioSink() override { return nullptr; }
+		virtual IPixelStreaming2VideoSink*					 GetPeerVideoSink(FString PlayerId) override { return nullptr; }
+		virtual IPixelStreaming2VideoSink*					 GetUnwatchedVideoSink() override { return nullptr; }
+		virtual void										 SetConfigOption(const FName& OptionName, const FString& Value) override {}
+		virtual bool										 GetConfigOption(const FName& OptionName, FString& OutValue) override { return false; }
+		virtual void										 PlayerRequestsBitrate(FString PlayerId, int MinBitrate, int MaxBitrate) override {}
+		virtual void										 RefreshStreamBitrate() override {}
+
+	private:
+		FPreConnectionEvent	   StreamingPreConnectionEvent;
+		FStreamingStartedEvent StreamingStartedEvent;
+		FStreamingStoppedEvent StreamingStoppedEvent;
+	};
+
+	FString FDummyStreamerFactory::GetStreamType()
+	{
+		return FString();
+	}
+
+	TSharedPtr<IPixelStreaming2Streamer> FDummyStreamerFactory::CreateNewStreamer(const FString& StreamerId)
+	{
+		return TSharedPtr<FDummyStreamer>(new FDummyStreamer());
+	}
+} // namespace UE::PixelStreaming2
+
+IPixelStreaming2StreamerFactory::IPixelStreaming2StreamerFactory()
+{
+	IPixelStreaming2StreamerFactory::RegisterStreamerFactory(this);
+}
+
+IPixelStreaming2StreamerFactory::~IPixelStreaming2StreamerFactory()
+{
+	IPixelStreaming2StreamerFactory::UnregisterStreamerFactory(this);
+}
+
+void IPixelStreaming2StreamerFactory::RegisterStreamerFactory(IPixelStreaming2StreamerFactory* InFactory)
+{
+	IModularFeatures::Get().RegisterModularFeature(ModularFeatureName, InFactory);
+}
+
+void IPixelStreaming2StreamerFactory::UnregisterStreamerFactory(IPixelStreaming2StreamerFactory* InFactory)
+{
+	IModularFeatures::Get().UnregisterModularFeature(ModularFeatureName, InFactory);
+}
+
+IPixelStreaming2StreamerFactory* IPixelStreaming2StreamerFactory::Get(const FString& InType)
+{
+	if (InType == FString())
+	{
+		return nullptr;
+	}
+
+	IModularFeatures::Get().LockModularFeatureList();
+	TArray<IPixelStreaming2StreamerFactory*> Factories = IModularFeatures::Get().GetModularFeatureImplementations<IPixelStreaming2StreamerFactory>(ModularFeatureName);
+	IModularFeatures::Get().UnlockModularFeatureList();
+
+	for (IPixelStreaming2StreamerFactory* Factory : Factories)
+	{
+		if (Factory && InType == Factory->GetStreamType())
+		{
+			return Factory;
+		}
+	}
+
+	UE_LOGFMT(LogPixelStreaming2Core, Warning, "No streamer factory implementation for {0} found. Streamers set to this type will not do anything.", *InType);
+	static UE::PixelStreaming2::FDummyStreamerFactory DummyFactory = UE::PixelStreaming2::FDummyStreamerFactory();
+	return &DummyFactory;
+}
+
+TArray<FString> IPixelStreaming2StreamerFactory::GetAvailableFactoryTypes()
+{
+	TArray<FString> StreamProtocols;
+
+	IModularFeatures::Get().LockModularFeatureList();
+	TArray<IPixelStreaming2StreamerFactory*> Factories = IModularFeatures::Get().GetModularFeatureImplementations<IPixelStreaming2StreamerFactory>(ModularFeatureName);
+	IModularFeatures::Get().UnlockModularFeatureList();
+
+	for (IPixelStreaming2StreamerFactory* Factory : Factories)
+	{
+		StreamProtocols.AddUnique(Factory->GetStreamType());
+	}
+
+	return StreamProtocols;
+}
\ No newline at end of file
diff --git a/Engine/Plugins/Media/PixelStreaming2/Source/PixelStreaming2Core/Private/Streamer.h b/Engine/Plugins/Media/PixelStreaming2/Source/PixelStreaming2Core/Private/Streamer.h
new file mode 100644
--- /dev/null
+++ b/Engine/Plugins/Media/PixelStreaming2/Source/PixelStreaming2Core/Private/Streamer.h
@@ -0,0 +1,13 @@
+// Copyright Epic Games, Inc. All Rights Reserved.
+
+#include "IPixelStreaming2Streamer.h"
+
+namespace UE::PixelStreaming2
+{
+	class FDummyStreamerFactory : public IPixelStreaming2StreamerFactory
+	{
+	public:
+		virtual FString								 GetStreamType() override;
+		virtual TSharedPtr<IPixelStreaming2Streamer> CreateNewStreamer(const FString& StreamerId) override;
+	};
+} // namespace UE::PixelStreaming2
\ No newline at end of file
diff --git a/Engine/Plugins/Media/PixelStreaming2/Source/PixelStreaming2Core/Public/IPixelStreaming2AudioConsumer.h b/Engine/Plugins/Media/PixelStreaming2/Source/PixelStreaming2Core/Public/IPixelStreaming2AudioConsumer.h
new file mode 100644
--- /dev/null
+++ b/Engine/Plugins/Media/PixelStreaming2/Source/PixelStreaming2Core/Public/IPixelStreaming2AudioConsumer.h
@@ -0,0 +1,43 @@
+// Copyright Epic Games, Inc. All Rights Reserved.
+
+#pragma once
+
+#include "CoreTypes.h"
+#include "UObject/Interface.h"
+
+#include "IPixelStreaming2AudioConsumer.generated.h"
+
+UINTERFACE(MinimalAPI)
+class UPixelStreaming2AudioConsumer : public UInterface
+{
+	GENERATED_BODY()
+};
+
+/**
+ * An "Audio Consumer" is an object that is responsible for outputting the audio received from a peer. For example, by
+ * passing the audio into a UE submix.
+ */
+class PIXELSTREAMING2CORE_API IPixelStreaming2AudioConsumer
+{
+	GENERATED_BODY()
+
+public:
+	/**
+	 * @brief Consume raw audio data.
+	 * @param AudioData Pointer to the audio data.
+	 * @param InSampleRate Audio sample rate in samples per second.
+	 * @param NChannels Number of audio channels. For example 2 for stero audio.
+	 * @param NFrames Number of Audio frames in a single channel.
+	 */
+	virtual void ConsumeRawPCM(const int16_t* AudioData, int InSampleRate, size_t NChannels, size_t NFrames) = 0;
+
+	/**
+	 * @brief Called when a audio consumer is added.
+	 */
+	virtual void OnConsumerAdded() = 0;
+
+	/**
+	 * @brief Called when a audio consumer is removed.
+	 */
+	virtual void OnConsumerRemoved() = 0;
+};
diff --git a/Engine/Plugins/Media/PixelStreaming2/Source/PixelStreaming2Core/Public/IPixelStreaming2AudioProducer.h b/Engine/Plugins/Media/PixelStreaming2/Source/PixelStreaming2Core/Public/IPixelStreaming2AudioProducer.h
new file mode 100644
--- /dev/null
+++ b/Engine/Plugins/Media/PixelStreaming2/Source/PixelStreaming2Core/Public/IPixelStreaming2AudioProducer.h
@@ -0,0 +1,38 @@
+// Copyright Epic Games, Inc. All Rights Reserved.
+
+#pragma once
+
+#include "CoreTypes.h"
+#include "UObject/Interface.h"
+
+#include "IPixelStreaming2AudioProducer.generated.h"
+
+UINTERFACE(MinimalAPI)
+class UPixelStreaming2AudioProducer : public UInterface
+{
+	GENERATED_BODY()
+};
+
+/**
+ * A "Audio Producer" is an object that you use to push audio frames into the Pixel Streaming system.
+ *
+ * Example usage:
+ *
+ * (1) Each new frame you want to push you call `MyAudioProducer->PushFrame(Buffer, NumSamples, NumChannels, SampleRate)`
+ *
+ * Note: Any audio you push in with `PushAudio` will be mixed with other audio producers before being streamed
+ */
+class PIXELSTREAMING2CORE_API IPixelStreaming2AudioProducer
+{
+	GENERATED_BODY()
+
+public:
+	/**
+	 * @brief Pushes audio into the PS2 audio pipeline which will mix with other audio producers before broadcasting.
+	 * @param InBuffer Pointer to the audio data.
+	 * @param NumSamples Number of Audio frames in a single channel.
+	 * @param NumChannels Number of audio channels. For example 2 for stero audio.
+	 * @param SampleRate Audio sample rate in samples per second.
+	 */
+	virtual void PushAudio(const float* InBuffer, int32 NumSamples, int32 NumChannels, int32 SampleRate) = 0;
+};
diff --git a/Engine/Plugins/Media/PixelStreaming2/Source/PixelStreaming2Core/Public/IPixelStreaming2AudioSink.h b/Engine/Plugins/Media/PixelStreaming2/Source/PixelStreaming2Core/Public/IPixelStreaming2AudioSink.h
new file mode 100644
--- /dev/null
+++ b/Engine/Plugins/Media/PixelStreaming2/Source/PixelStreaming2Core/Public/IPixelStreaming2AudioSink.h
@@ -0,0 +1,37 @@
+// Copyright Epic Games, Inc. All Rights Reserved.
+
+#pragma once
+
+#include "CoreTypes.h"
+#include "IPixelStreaming2AudioConsumer.h"
+#include "UObject/Interface.h"
+
+#include "IPixelStreaming2AudioSink.generated.h"
+
+UINTERFACE(MinimalAPI)
+class UPixelStreaming2AudioSink : public UInterface
+{
+	GENERATED_BODY()
+};
+
+/**
+ * An "Audio Sink" is an object that receives audio from a singular peer. From here, you can add multiple consumers to
+ * output the received audio.
+ */
+class PIXELSTREAMING2CORE_API IPixelStreaming2AudioSink
+{
+	GENERATED_BODY()
+
+public:
+	/**
+	 * @brief Add an audio consumer to the sink.
+	 * @param AudioConsumer The Audio consumer to add to the sink.
+	 */
+	virtual void AddAudioConsumer(IPixelStreaming2AudioConsumer* AudioConsumer) = 0;
+
+	/**
+	 * @brief Remove an audio consumer to remove from the sink.
+	 * @param AudioConsumer The Audio consumer to remove from the sink.
+	 */
+	virtual void RemoveAudioConsumer(IPixelStreaming2AudioConsumer* AudioConsumer) = 0;		
+};
\ No newline at end of file
diff --git a/Engine/Plugins/Media/PixelStreaming2/Source/PixelStreaming2Core/Public/IPixelStreaming2Streamer.h b/Engine/Plugins/Media/PixelStreaming2/Source/PixelStreaming2Core/Public/IPixelStreaming2Streamer.h
new file mode 100644
--- /dev/null
+++ b/Engine/Plugins/Media/PixelStreaming2/Source/PixelStreaming2Core/Public/IPixelStreaming2Streamer.h
@@ -0,0 +1,294 @@
+// Copyright Epic Games, Inc. All Rights Reserved.
+
+#pragma once
+
+#include "CoreTypes.h"
+#include "Features/IModularFeature.h"
+#include "IPixelStreaming2AudioSink.h"
+#include "IPixelStreaming2VideoProducer.h"
+#include "IPixelStreaming2VideoSink.h"
+#include "Templates/SharedPointer.h"
+#include "UObject/Interface.h"
+
+class UTexture2D;
+namespace UE::PixelStreaming2
+{
+	class FPixelStreaming2Module;
+} // namespace UE::PixelStreaming2
+
+/**
+ * The IPixelStreaming2Streamer interface provides functionality for managing a Pixel Streaming session.
+ * This class manages the core operations of streaming, such as setting the framerate, handling video input,
+ * and interacting with the signalling server. It also allows for control over connected players and sending data to them.
+ */
+class PIXELSTREAMING2CORE_API IPixelStreaming2Streamer
+{
+public:
+	virtual ~IPixelStreaming2Streamer() = default;
+
+	/**
+	 * Called just after streamer creation. Use this method to initialize any methods that need
+	 * a shared pointer to the streamer by using AsShared() inside the method.
+	 */
+	virtual void Initialize() = 0;
+
+	/**
+	 * @brief Set the Stream FPS.
+	 * @param InFramesPerSecond The number of frames per second the streamer will stream at.
+	 */
+	virtual void SetStreamFPS(int32 InFramesPerSecond) = 0;
+
+	/**
+	 * @brief Get the Stream FPS.
+	 * @return - The number of frames per second the streamer will stream at.
+	 */
+	virtual int32 GetStreamFPS() = 0;
+
+	/**
+	 * Setting this to true will cause the streamer to ignore the FPS value and instead push out frames
+	 * as they are submitted from the video input. If the encode process takes longer than the time
+	 * between frames, frames will be dropped.
+	 * @param bCouple true to couple the streamer framerate to that of the video input.
+	 */
+	virtual void SetCoupleFramerate(bool bCouple) = 0;
+
+	/**
+	 * @brief Set the Video Input object.
+	 * @param Input The IPixelStreaming2VideoProducer that this streamer will stream.
+	 */
+	virtual void SetVideoProducer(TSharedPtr<IPixelStreaming2VideoProducer> Input) = 0;
+
+	/**
+	 * @brief Get the Video Input object.
+	 * @return The IPixelStreaming2VideoProducer that this streamer will stream.
+	 */
+	virtual TWeakPtr<IPixelStreaming2VideoProducer> GetVideoProducer() = 0;
+
+	/**
+	 * @brief Set the URL this streamer will connect to.
+	 * @param InConnectionURL.
+	 */
+	virtual void SetConnectionURL(const FString& InConnectionURL) = 0;
+
+	/**
+	 * @brief Set the Signalling Server URL.
+	 * @param InSignallingServerURL.
+	 */
+	UE_DEPRECATED(5.6, "SetSignallingServerURL has been deprecated. Please use SetConnectionURL instead.")
+	virtual void SetSignallingServerURL(const FString& InSignallingServerURL) { SetConnectionURL(InSignallingServerURL); }
+
+	/**
+	 * @brief Get the URL this streamer will connect to.
+	 * @return The connection URL.
+	 */
+	virtual FString GetConnectionURL() = 0;
+
+	/**
+	 * @brief Get the Signalling Server URL.
+	 * @return The Signalling Server URL.
+	 */
+	UE_DEPRECATED(5.6, "GetSignallingServerURL has been deprecated. Please use GetConnectionURL instead.")
+	virtual FString GetSignallingServerURL() { return GetConnectionURL(); }
+
+	/**
+	 * @brief Get this streamer's ID.
+	 * @return The streamer's ID.
+	 */
+	virtual FString GetId() = 0;
+
+	/**
+	 * @brief Check if this streamer is currently connected to the connection URL.
+	 * @return True if the streamer is connected.
+	 */
+	virtual bool IsConnected() = 0;
+
+	/**
+	 * @brief Check if this streamer is currently connected to the signalling mechanism (e.g. websocket for signalling server).
+	 * @return True if the streamer is connected to the signalling mechanism.
+	 */
+	UE_DEPRECATED(5.6, "IsSignallingConnected has been deprecated. Please use IsConnected instead.")
+	virtual bool IsSignallingConnected() { return IsConnected(); }
+
+	/**
+	 * @brief Start streaming this streamer
+	 */
+	virtual void StartStreaming() = 0;
+
+	/**
+	 * @brief Stop this streamer from streaming
+	 */
+	virtual void StopStreaming() = 0;
+
+	/**
+	 * @brief Get the current state of this streamer
+	 * @return True if streaming, false otherwise.
+	 */
+	virtual bool IsStreaming() const = 0;
+
+	/**
+	 * @brief Event fired just before the streamer begins connecting to signalling.
+	 */
+	DECLARE_EVENT_OneParam(IPixelStreaming2Streamer, FPreConnectionEvent, IPixelStreaming2Streamer*);
+
+	/**
+	 * @brief A getter for the OnPreConnection event. Intent is for users to call IPixelStreaming2Module::Get().FindStreamer(ID)->OnPreConnection().AddXXX.
+	 * @return - The bindable OnPreConnection event.
+	 */
+	virtual FPreConnectionEvent& OnPreConnection() = 0;
+
+	/**
+	 * @brief Event fired when the streamer has connected to a signalling server and is ready for peers.
+	 */
+	DECLARE_EVENT_OneParam(IPixelStreaming2Streamer, FStreamingStartedEvent, IPixelStreaming2Streamer*);
+
+	/**
+	 * @brief A getter for the OnStreamingStarted event. Intent is for users to call IPixelStreaming2Module::Get().FindStreamer(ID)->OnStreamingStarted().AddXXX.
+	 * @return - The bindable OnStreamingStarted event.
+	 */
+	virtual FStreamingStartedEvent& OnStreamingStarted() = 0;
+
+	/**
+	 * @brief Event fired when the streamer has disconnected from a signalling server and has stopped streaming.
+	 */
+	DECLARE_EVENT_OneParam(IPixelStreaming2Streamer, FStreamingStoppedEvent, IPixelStreaming2Streamer*);
+
+	/**
+	 * @brief A getter for the OnStreamingStopped event. Intent is for users to call IPixelStreaming2Module::Get().FindStreamer(ID)->OnStreamingStopped().AddXXX.
+	 * @return - The bindable OnStreamingStopped event.
+	 */
+	virtual FStreamingStoppedEvent& OnStreamingStopped() = 0;
+
+	/**
+	 * @brief Force a key frame to be sent.
+	 */
+	virtual void ForceKeyFrame() = 0;
+
+	/**
+	 * @brief Freeze Pixel Streaming.
+	 * @param Texture 		- The freeze frame to display. If null then the back buffer is captured.
+	 */
+	virtual void FreezeStream(UTexture2D* Texture) = 0;
+
+	/**
+	 * @brief Unfreeze Pixel Streaming.
+	 */
+	virtual void UnfreezeStream() = 0;
+
+	/**
+	 * @brief Send all players connected to this streamer a message.
+	 * @param MessageType The message type to be sent to the player.
+	 * @param Descriptor The contents of the message.
+	 */
+	virtual void SendAllPlayersMessage(FString MessageType, const FString& Descriptor) = 0;
+
+	/**
+	 * @brief Send all players connected to this streamer a message.
+	 * @param MessageType The message type to be sent to the player.
+	 * @param Descriptor The contents of the message.
+	 */
+	virtual void SendPlayerMessage(FString PlayerId, FString MessageType, const FString& Descriptor) = 0;
+
+	/**
+	 * @brief Send a file to the browser where we are sending video.
+	 * @param ByteData	 	- The raw byte data of the file.
+	 * @param MimeType	 	- The files Mime type. Used for reconstruction on the front end.
+	 * @param FileExtension - The files extension. Used for reconstruction on the front end.
+	 */
+	virtual void SendFileData(const TArray64<uint8>& ByteData, FString& MimeType, FString& FileExtension) = 0;
+
+	/**
+	 * @brief Kick a player by player id.
+	 * @param PlayerId		- The ID of the player to kick
+	 */
+	virtual void KickPlayer(FString PlayerId) = 0;
+
+	/**
+	 * @brief Get the ids of the connected players.
+	 * @return TArray<FString> The ids of the connected players.
+	 */
+	virtual TArray<FString> GetConnectedPlayers() = 0;
+
+	/**
+	 * @brief Get the streamer's input handler
+	 * @return The streamer's input handler
+	 */
+	virtual TWeakPtr<class IPixelStreaming2InputHandler> GetInputHandler() = 0;
+
+	/**
+	 * @brief Get the audio sink associated with a specific peer/player.
+	 * @param PlayerId The player id of the audio sink to retrieve.
+	 * @return The PeerAudioSink.
+	 */
+	virtual IPixelStreaming2AudioSink* GetPeerAudioSink(FString PlayerId) = 0;
+
+	/**
+	 * @brief Get an audio sink that has no peers/players listening to it.
+	 * @return The unlistened PeerAudioSink.
+	 */
+	virtual IPixelStreaming2AudioSink* GetUnlistenedAudioSink() = 0;
+
+	/**
+	 * @brief Get the video sink associated with a specific peer/player.
+	 * @return The PeerVideoSink.
+	 */
+	virtual IPixelStreaming2VideoSink* GetPeerVideoSink(FString PlayerId) = 0;
+
+	/**
+	 * @brief Get a video sink that has no peers/players watching it.
+	 * @return The unwatched PeerVideoSink.
+	 */
+	virtual IPixelStreaming2VideoSink* GetUnwatchedVideoSink() = 0;
+
+	/**
+	 * @brief Allows sending arbitrary configuration options during initial connection.
+	 * @param OptionName The name of the option to set.
+	 * @param Value Setting a value to an empty string clears it from the mapping and prevents it being sent.
+	 */
+	virtual void SetConfigOption(const FName& OptionName, const FString& Value) = 0;
+
+	/**
+	 * @brief Get the Gonfigerartion value for a specific option.
+	 * @param OptionName The name of the config option to get.
+	 * @param Value The value of the config option retrieved.
+	 */
+	virtual bool GetConfigOption(const FName& OptionName, FString& OutValue) = 0;
+
+	/**
+	 * @brief Set the minimum and maximum bitrate for the streamer.
+	 * @param PlayerId Currently unused. For setting the bitrate for a individual player id.
+	 * @param MinBitrate minimum bitrate for the streamer.
+	 * @param MaxBitrate maximum bitrate for the streamer.
+	 */
+	virtual void PlayerRequestsBitrate(FString PlayerId, int MinBitrate, int MaxBitrate) = 0;
+
+	/**
+	 * @brief Refresh connection with minimum and maximum bitrate.
+	 */
+	virtual void RefreshStreamBitrate() = 0;
+};
+
+/**
+ * The IPixelStreaming2StreamerFactory interface provides functionality for creating custom streamers based on an identifier.
+ *
+ */
+class PIXELSTREAMING2CORE_API IPixelStreaming2StreamerFactory : public IModularFeature
+{
+public:
+	IPixelStreaming2StreamerFactory();
+	virtual ~IPixelStreaming2StreamerFactory();
+
+	virtual FString GetStreamType() = 0;
+
+	static void RegisterStreamerFactory(IPixelStreaming2StreamerFactory* InFactory);
+
+	static void UnregisterStreamerFactory(IPixelStreaming2StreamerFactory* InFactory);
+
+	static IPixelStreaming2StreamerFactory* Get(const FString& InType);
+
+	static TArray<FString> GetAvailableFactoryTypes();
+
+private:
+	// Only the module should call CreateNewStreamer on the factories.
+	friend UE::PixelStreaming2::FPixelStreaming2Module;
+	virtual TSharedPtr<IPixelStreaming2Streamer> CreateNewStreamer(const FString& StreamerId) = 0;
+};
diff --git a/Engine/Plugins/Media/PixelStreaming2/Source/PixelStreaming2Core/Public/IPixelStreaming2VideoConsumer.h b/Engine/Plugins/Media/PixelStreaming2/Source/PixelStreaming2Core/Public/IPixelStreaming2VideoConsumer.h
new file mode 100644
--- /dev/null
+++ b/Engine/Plugins/Media/PixelStreaming2/Source/PixelStreaming2Core/Public/IPixelStreaming2VideoConsumer.h
@@ -0,0 +1,41 @@
+// Copyright Epic Games, Inc. All Rights Reserved.
+
+#pragma once
+
+#include "CoreTypes.h"
+#include "UObject/Interface.h"
+#include "RHIFwd.h"
+
+#include "IPixelStreaming2VideoConsumer.generated.h"
+
+UINTERFACE(MinimalAPI)
+class UPixelStreaming2VideoConsumer : public UInterface
+{
+	GENERATED_BODY()
+};
+
+/**
+ * An "Video Consumer" is an object that is responsible for outputting the video received from a peer. For example, by
+ * rendering to a render target.
+ */
+class PIXELSTREAMING2CORE_API IPixelStreaming2VideoConsumer
+{
+	GENERATED_BODY()
+
+public:
+	/**
+	 * @brief Consume a texture as a video frame.
+	 * @param Frame The Frame to consume.
+	 */
+	virtual void ConsumeFrame(FTextureRHIRef Frame) = 0;
+
+	/**
+	 * @brief Called when a video consumer is added.
+	 */
+	virtual void OnConsumerAdded() = 0;
+
+	/**
+	 * @brief Called when a video consumer is removed.
+	 */
+	virtual void OnConsumerRemoved() = 0;
+};
diff --git a/Engine/Plugins/Media/PixelStreaming2/Source/PixelStreaming2Core/Public/IPixelStreaming2VideoProducer.h b/Engine/Plugins/Media/PixelStreaming2/Source/PixelStreaming2Core/Public/IPixelStreaming2VideoProducer.h
new file mode 100644
--- /dev/null
+++ b/Engine/Plugins/Media/PixelStreaming2/Source/PixelStreaming2Core/Public/IPixelStreaming2VideoProducer.h
@@ -0,0 +1,43 @@
+// Copyright Epic Games, Inc. All Rights Reserved.
+
+#pragma once
+
+#include "CoreTypes.h"
+#include "IPixelCaptureInputFrame.h"
+#include "UObject/Interface.h"
+
+#include "IPixelStreaming2VideoProducer.generated.h"
+
+UINTERFACE(MinimalAPI)
+class UPixelStreaming2VideoProducer : public UInterface
+{
+	GENERATED_BODY()
+};
+
+/**
+ * A "Video Producer" is an object that you use to push video frames into the Pixel Streaming system.
+ *
+ * Example usage:
+ *
+ * (1) Each new frame you want to push you call `MyVideoProducer->PushFrame(MyFrame)`
+ *
+ * Note: Your frame is likely a `FPixelCaptureInputFrameRHI` if the frame is a GPU texture or
+ * a `FPixelCaptureInputFrameI420` or `FPixelCaptureInputFrameNV12` if your frame is from the CPU.
+ */
+class PIXELSTREAMING2CORE_API IPixelStreaming2VideoProducer
+{
+	GENERATED_BODY()
+
+public:
+	/**
+	 * Pushes a raw video frame into the Pixel Streaming system.
+	 * @param InputFrame The raw input frame, which may be a GPU texture or a CPU texture based on the underlying type of frame pushed.
+	 */
+	virtual void PushFrame(const IPixelCaptureInputFrame& InputFrame) = 0;
+
+    /**
+	 * A human readable identifier used when displaying what the streamer is streaming in the toolbar
+	 * @return A string containing the display name.
+	 */
+	virtual FString ToString() = 0;
+};
\ No newline at end of file
diff --git a/Engine/Plugins/Media/PixelStreaming2/Source/PixelStreaming2Core/Public/IPixelStreaming2VideoSink.h b/Engine/Plugins/Media/PixelStreaming2/Source/PixelStreaming2Core/Public/IPixelStreaming2VideoSink.h
new file mode 100644
--- /dev/null
+++ b/Engine/Plugins/Media/PixelStreaming2/Source/PixelStreaming2Core/Public/IPixelStreaming2VideoSink.h
@@ -0,0 +1,38 @@
+// Copyright Epic Games, Inc. All Rights Reserved.
+
+#pragma once
+
+#include "CoreTypes.h"
+#include "IPixelStreaming2VideoConsumer.h"
+#include "UObject/Interface.h"
+
+#include "IPixelStreaming2VideoSink.generated.h"
+
+UINTERFACE(MinimalAPI)
+class UPixelStreaming2VideoSink : public UInterface
+{
+	GENERATED_BODY()
+};
+
+class IPixelStreaming2VideoConsumer;
+
+/**
+ * Interface for a sink that collects video coming in from the browser and passes into into UE.
+ */
+class PIXELSTREAMING2CORE_API IPixelStreaming2VideoSink
+{
+	GENERATED_BODY()
+
+public:
+	/**
+	 * @brief Add a video consumer to the sink.
+	 * @param VideoConsumer The video consumer to add to the sink.
+	 */
+	virtual void AddVideoConsumer(IPixelStreaming2VideoConsumer* VideoConsumer) = 0;
+
+	/**
+	 * @brief Remove a video consumer to remove from the sink.
+	 * @param VideoConsumer The video consumer to remove from the sink.
+	 */
+	virtual void RemoveVideoConsumer(IPixelStreaming2VideoConsumer* VideoConsumer) = 0;
+};
diff --git a/Engine/Plugins/Media/PixelStreaming2/Source/PixelStreaming2Editor/PixelStreaming2Editor.build.cs b/Engine/Plugins/Media/PixelStreaming2/Source/PixelStreaming2Editor/PixelStreaming2Editor.build.cs
--- a/Engine/Plugins/Media/PixelStreaming2/Source/PixelStreaming2Editor/PixelStreaming2Editor.build.cs
+++ b/Engine/Plugins/Media/PixelStreaming2/Source/PixelStreaming2Editor/PixelStreaming2Editor.build.cs
@@ -33,7 +33,8 @@ namespace UnrealBuildTool.Rules
 				"ApplicationCore",
 				"PixelStreaming2Input",
 				"AVCodecsCore",
-				"PixelStreaming2"
+				"PixelStreaming2",
+				"PixelStreaming2Core"
 			});
 
 			if (Target.bBuildEditor)
diff --git a/Engine/Plugins/Media/PixelStreaming2/Source/PixelStreaming2RTC/PixelStreaming2RTC.Build.cs b/Engine/Plugins/Media/PixelStreaming2/Source/PixelStreaming2RTC/PixelStreaming2RTC.Build.cs
--- a/Engine/Plugins/Media/PixelStreaming2/Source/PixelStreaming2RTC/PixelStreaming2RTC.Build.cs
+++ b/Engine/Plugins/Media/PixelStreaming2/Source/PixelStreaming2RTC/PixelStreaming2RTC.Build.cs
@@ -69,7 +69,8 @@ namespace UnrealBuildTool.Rules
                 "PixelStreaming2Input",
                 "EpicRtc",
                 "AVCodecsCore",
-                "PixelStreaming2"
+                "PixelStreaming2",
+                "PixelStreaming2Core"
             });
 
             PrivateDependencyModuleNames.AddRange(new string[]
diff --git a/Engine/Plugins/Media/PixelStreaming2/Source/PixelStreaming2Settings/Internal/PixelStreaming2PluginSettings.h b/Engine/Plugins/Media/PixelStreaming2/Source/PixelStreaming2Settings/Internal/PixelStreaming2PluginSettings.h
--- a/Engine/Plugins/Media/PixelStreaming2/Source/PixelStreaming2Settings/Internal/PixelStreaming2PluginSettings.h
+++ b/Engine/Plugins/Media/PixelStreaming2/Source/PixelStreaming2Settings/Internal/PixelStreaming2PluginSettings.h
@@ -140,7 +140,8 @@ public:
 	static TAutoConsoleVariable<FString> CVarDefaultStreamerType;
 	UPROPERTY(config, EditAnywhere, Category = "PixelStreaming", meta = (
 		DisplayName = "Default Streamer Type",
-		ToolTip = "Default Streamer Type to be used when not specified elsewhere. This value should match a type registered by an IPixelStreaming2StreamerFactory"
+		ToolTip = "Default Streamer Type to be used when not specified elsewhere. This value should match a type registered by an IPixelStreaming2StreamerFactory",
+		GetOptions = "GetDefaultStreamerTypeOptions"
 		))
 	FString DefaultStreamerType = TEXT("DefaultRtc");
 
@@ -626,4 +627,7 @@ private:
 
 	UFUNCTION()
 	TArray<FString> GetWebRTCCodecPreferencesOptions() const;
+
+	UFUNCTION()
+	TArray<FString> GetDefaultStreamerTypeOptions() const;
 };
\ No newline at end of file
diff --git a/Engine/Plugins/Media/PixelStreaming2/Source/PixelStreaming2Settings/PixelStreaming2Settings.Build.cs b/Engine/Plugins/Media/PixelStreaming2/Source/PixelStreaming2Settings/PixelStreaming2Settings.Build.cs
--- a/Engine/Plugins/Media/PixelStreaming2/Source/PixelStreaming2Settings/PixelStreaming2Settings.Build.cs
+++ b/Engine/Plugins/Media/PixelStreaming2/Source/PixelStreaming2Settings/PixelStreaming2Settings.Build.cs
@@ -16,6 +16,7 @@ namespace UnrealBuildTool.Rules
 				"DeveloperSettings",
 				"Engine",
 				"EngineSettings",
+				"PixelStreaming2Core"
 			});
 
 			PublicDependencyModuleNames.AddRange(new string[]
diff --git a/Engine/Plugins/Media/PixelStreaming2/Source/PixelStreaming2Settings/Private/PixelStreaming2PluginSettings.cpp b/Engine/Plugins/Media/PixelStreaming2/Source/PixelStreaming2Settings/Private/PixelStreaming2PluginSettings.cpp
--- a/Engine/Plugins/Media/PixelStreaming2/Source/PixelStreaming2Settings/Private/PixelStreaming2PluginSettings.cpp
+++ b/Engine/Plugins/Media/PixelStreaming2/Source/PixelStreaming2Settings/Private/PixelStreaming2PluginSettings.cpp
@@ -2,6 +2,7 @@
 
 #include "PixelStreaming2PluginSettings.h"
 
+#include "IPixelStreaming2Streamer.h"
 #include "Logging.h"
 #include "Misc/CommandLine.h"
 #include "UObject/ReflectedTypeAccessors.h"
@@ -59,6 +60,34 @@ namespace
 		}
 	}
 
+	static void VerifyCVarDefaultStreamerType(IConsoleVariable* CVar)
+	{
+		TArray<FString> AvailableFactoryTypes = IPixelStreaming2StreamerFactory::GetAvailableFactoryTypes();
+		FString SpecifiedFactory = CVar->GetString();
+
+		if(AvailableFactoryTypes.Num() == 0)
+		{
+			// This code path executes when the cvar is initially set and no factories have been registered
+			return;
+		}
+
+		bool bValid = false;
+		for(const FString& AvailableFactory : AvailableFactoryTypes)
+		{
+			if(SpecifiedFactory == AvailableFactory)
+			{
+				bValid =  true;
+				break;
+			}
+		}
+
+		if(!bValid)
+		{
+			UE_LOGFMT(LogPixelStreaming2Settings, Warning, "\"{0}\" isn't a registered streamer type. Valid types: [{1}]. Restoring to \"{2}\"", SpecifiedFactory, FString::Join(AvailableFactoryTypes, TEXT(",")), CVar->GetDefaultValue());
+			CVar->SetWithCurrentPriority(*CVar->GetDefaultValue());
+		}
+	}
+
 	FString ConsoleVariableToCommandArgValue(const FString InCVarName)
 	{
 		// CVars are . deliminated by section. To get their equivilent commandline arg for parsing
@@ -274,6 +303,7 @@ TAutoConsoleVariable<FString> UPixelStreaming2PluginSettings::CVarDefaultStreame
 	TEXT("PixelStreaming2.DefaultStreamerType"),
 	TEXT("DefaultRtc"),
 	TEXT("Default Streamer Type to be used when not specified elsewhere."),
+	FConsoleVariableDelegate::CreateLambda([](IConsoleVariable* Var) { VerifyCVarDefaultStreamerType(Var); }),
 	ECVF_Default);
 
 TAutoConsoleVariable<FString> UPixelStreaming2PluginSettings::CVarConnectionURL(
@@ -1372,4 +1402,9 @@ TArray<FString> UPixelStreaming2PluginSettings::GetWebRTCCodecPreferencesOptions
 	}
 
 	return PossibleCodecs.Array();
+}
+
+TArray<FString> UPixelStreaming2PluginSettings::GetDefaultStreamerTypeOptions() const
+{
+	return IPixelStreaming2StreamerFactory::GetAvailableFactoryTypes();
 }
\ No newline at end of file
