[PS2] QOL: Move required PixelStreaming2RTC Private headers to Internal for other use cases

diff --git a/Engine/Plugins/Media/PixelStreaming2/Source/PixelStreaming2/Internal/AudioCapturer.h b/Engine/Plugins/Media/PixelStreaming2/Source/PixelStreaming2/Internal/AudioCapturer.h
--- a/Engine/Plugins/Media/PixelStreaming2/Source/PixelStreaming2/Internal/AudioCapturer.h
+++ b/Engine/Plugins/Media/PixelStreaming2/Source/PixelStreaming2/Internal/AudioCapturer.h
@@ -5,7 +5,6 @@
 #include "DSP/MultithreadedPatching.h"
 #include "AudioProducer.h"
 #include "ISubmixBufferListener.h"
-#include "Logging.h"
 #include "SampleBuffer.h"
 #include "TickableTask.h"
 
diff --git a/Engine/Plugins/Media/PixelStreaming2/Source/PixelStreaming2/Internal/StreamerReconnectTimer.h b/Engine/Plugins/Media/PixelStreaming2/Source/PixelStreaming2/Internal/StreamerReconnectTimer.h
new file mode 100644
--- /dev/null
+++ b/Engine/Plugins/Media/PixelStreaming2/Source/PixelStreaming2/Internal/StreamerReconnectTimer.h
@@ -0,0 +1,40 @@
+// Copyright Epic Games, Inc. All Rights Reserved.
+
+#pragma once
+
+#include "Templates/SharedPointer.h"
+#include "Tickable.h"
+#include "IPixelStreaming2Streamer.h"
+
+namespace UE::PixelStreaming2
+{
+	/**
+	 * An interval timer that triggers reconnects when `IPixelStreaming2Streamer` has it session disconnect.
+	 * This is different than a websocket error reconnect because in this case we were already connected.
+	 * Timer takes a `TWeakPtr<IPixelStreaming2Streamer>`, meaning the streamer object must exist for the timer to do anything.
+	 * Note: this timer runs on the game thread and its reconnect interval is controlled by the CVar
+	 * UPixelStreaming2PluginSettings::CVarSignalingReconnectInterval
+	*/
+	class PIXELSTREAMING2_API FStreamerReconnectTimer : public FTickableGameObject
+	{
+	public:
+
+		FStreamerReconnectTimer();
+		virtual ~FStreamerReconnectTimer() = default;
+		void Start(TWeakPtr<IPixelStreaming2Streamer> InWeakStreamer);
+		void Stop();
+
+		// Begin FTickableGameObject
+		virtual void Tick (float DeltaTime) override;
+		virtual bool IsTickableInEditor() const override { return true; }
+		virtual bool IsTickableWhenPaused() const override { return true; }
+		FORCEINLINE TStatId GetStatId() const { RETURN_QUICK_DECLARE_CYCLE_STAT(PixelStreaming2StreamerReconnectTimer, STATGROUP_Tickables); }
+		// End FTickableGameObject
+
+	private:
+		TWeakPtr<IPixelStreaming2Streamer> WeakStreamer;
+		bool bEnabled = false;
+		uint64 LastReconnectCycles = 0;
+
+	};
+} // namespace UE::PixelStreaming2
\ No newline at end of file
diff --git a/Engine/Plugins/Media/PixelStreaming2/Source/PixelStreaming2/Private/StreamerReconnectTimer.cpp b/Engine/Plugins/Media/PixelStreaming2/Source/PixelStreaming2/Private/StreamerReconnectTimer.cpp
new file mode 100644
--- /dev/null
+++ b/Engine/Plugins/Media/PixelStreaming2/Source/PixelStreaming2/Private/StreamerReconnectTimer.cpp
@@ -0,0 +1,71 @@
+// Copyright Epic Games, Inc. All Rights Reserved.
+
+#include "StreamerReconnectTimer.h"
+
+#include "Logging.h"
+#include "PixelStreaming2PluginSettings.h"
+
+namespace UE::PixelStreaming2
+{
+
+	FStreamerReconnectTimer::FStreamerReconnectTimer()
+	{
+	}
+
+	void FStreamerReconnectTimer::Start(TWeakPtr<IPixelStreaming2Streamer> InWeakStreamer)
+	{
+		WeakStreamer = InWeakStreamer;
+		bEnabled = true;
+	}
+
+	void FStreamerReconnectTimer::Stop()
+	{
+		bEnabled = false;
+	}
+
+	void FStreamerReconnectTimer::Tick(float DeltaTime)
+	{
+		if (IsEngineExitRequested())
+		{
+			return;
+		}
+
+		if (!bEnabled)
+		{
+			return;
+		}
+
+		TSharedPtr<IPixelStreaming2Streamer> Streamer = WeakStreamer.Pin();
+
+		if (!Streamer)
+		{
+			return;
+		}
+
+		// Do not attempt a reconnect is we are already connected/streaming
+		if (Streamer->IsStreaming())
+		{
+			return;
+		}
+
+		float ReconnectInterval = UPixelStreaming2PluginSettings::CVarSignalingReconnectInterval.GetValueOnAnyThread();
+
+		if (ReconnectInterval <= 0.0f)
+		{
+			return;
+		}
+
+		uint64 CyclesNow = FPlatformTime::Cycles64();
+		uint64 DeltaCycles = CyclesNow - LastReconnectCycles;
+		float DeltaSeconds = FPlatformTime::ToSeconds(DeltaCycles);
+
+		// If enough time has elapsed, try a reconnect
+		if (DeltaSeconds >= ReconnectInterval)
+		{
+			UE_LOG(LogPixelStreaming2, Log, TEXT("Streamer reconnecting..."))
+				Streamer->StartStreaming();
+			LastReconnectCycles = CyclesNow;
+		}
+	}
+
+} // namespace UE::PixelStreaming2
\ No newline at end of file
diff --git a/Engine/Plugins/Media/PixelStreaming2/Source/PixelStreaming2RTC/Internal/EpicRtcAudioCapturer.h b/Engine/Plugins/Media/PixelStreaming2/Source/PixelStreaming2RTC/Internal/EpicRtcAudioCapturer.h
new file mode 100644
--- /dev/null
+++ b/Engine/Plugins/Media/PixelStreaming2/Source/PixelStreaming2RTC/Internal/EpicRtcAudioCapturer.h
@@ -0,0 +1,28 @@
+// Copyright Epic Games, Inc. All Rights Reserved.
+
+#pragma once
+
+#include "AudioCapturer.h"
+
+namespace UE::PixelStreaming2
+{
+	/**
+	 * FEpicRtcAudioCapturer overrides the default PushAudio behaviour of the FAudioCapturer in order to
+	 * break up the pushed audio into 10ms chunks
+	 */
+	class PIXELSTREAMING2RTC_API FEpicRtcAudioCapturer : public FAudioCapturer
+	{
+	public:
+		static TSharedPtr<FEpicRtcAudioCapturer> Create();
+		virtual ~FEpicRtcAudioCapturer() = default;
+
+		// Override the push audio method as EpicRtc needs the broadcasted audio to be in 10ms chunks
+		virtual void PushAudio(const float* AudioData, int32 NumSamples, int32 NumChannels, int32 SampleRate) override;
+
+	protected:
+		FEpicRtcAudioCapturer() = default;
+
+	private:
+		TArray<int16_t> RecordingBuffer;
+	};
+} // namespace UE::PixelStreaming2
\ No newline at end of file
diff --git a/Engine/Plugins/Media/PixelStreaming2/Source/PixelStreaming2RTC/Internal/EpicRtcAudioSink.h b/Engine/Plugins/Media/PixelStreaming2/Source/PixelStreaming2RTC/Internal/EpicRtcAudioSink.h
new file mode 100644
--- /dev/null
+++ b/Engine/Plugins/Media/PixelStreaming2/Source/PixelStreaming2RTC/Internal/EpicRtcAudioSink.h
@@ -0,0 +1,20 @@
+// Copyright Epic Games, Inc. All Rights Reserved.
+
+#pragma once
+
+#include "AudioSink.h"
+#include "EpicRtcAudioTrack.h"
+
+namespace UE::PixelStreaming2
+{
+	// Collects audio coming in from EpicRtc and passes into into UE's audio system.
+	class PIXELSTREAMING2RTC_API FEpicRtcAudioSink : public FAudioSink, public FEpicRtcAudioTrack
+	{
+	public:
+		static TSharedPtr<FEpicRtcAudioSink> Create(TRefCountPtr<EpicRtcAudioTrackInterface> InTrack);
+		virtual ~FEpicRtcAudioSink() = default;
+
+	private:
+		FEpicRtcAudioSink(TRefCountPtr<EpicRtcAudioTrackInterface> InTrack);
+	};
+} // namespace UE::PixelStreaming2
\ No newline at end of file
diff --git a/Engine/Plugins/Media/PixelStreaming2/Source/PixelStreaming2RTC/Internal/EpicRtcAudioSource.h b/Engine/Plugins/Media/PixelStreaming2/Source/PixelStreaming2RTC/Internal/EpicRtcAudioSource.h
new file mode 100644
--- /dev/null
+++ b/Engine/Plugins/Media/PixelStreaming2/Source/PixelStreaming2RTC/Internal/EpicRtcAudioSource.h
@@ -0,0 +1,23 @@
+// Copyright Epic Games, Inc. All Rights Reserved.
+
+#pragma once
+
+#include "AudioSource.h"
+#include "EpicRtcAudioTrack.h"
+#include "Templates/SharedPointer.h"
+
+#include "epic_rtc/core/audio/audio_track.h"
+
+namespace UE::PixelStreaming2
+{
+	class PIXELSTREAMING2RTC_API FEpicRtcAudioSource : public FAudioSource, public FEpicRtcAudioTrack
+	{
+	public:
+		static TSharedPtr<FEpicRtcAudioSource> Create(TRefCountPtr<EpicRtcAudioTrackInterface> InTrack, TSharedPtr<class FEpicRtcAudioCapturer> InCapturer);
+
+		virtual void OnAudioBuffer(const int16_t* AudioData, int32 NumSamples, int32 NumChannels, const int32 SampleRate);
+
+	private:
+		FEpicRtcAudioSource(TRefCountPtr<EpicRtcAudioTrackInterface> InTrack);
+	};
+} // namespace UE::PixelStreaming2
\ No newline at end of file
diff --git a/Engine/Plugins/Media/PixelStreaming2/Source/PixelStreaming2RTC/Internal/EpicRtcAudioTrack.h b/Engine/Plugins/Media/PixelStreaming2/Source/PixelStreaming2RTC/Internal/EpicRtcAudioTrack.h
new file mode 100644
--- /dev/null
+++ b/Engine/Plugins/Media/PixelStreaming2/Source/PixelStreaming2RTC/Internal/EpicRtcAudioTrack.h
@@ -0,0 +1,22 @@
+// Copyright Epic Games, Inc. All Rights Reserved.
+
+#pragma once
+
+#include "Templates/RefCounting.h"
+
+#include "epic_rtc/core/audio/audio_track.h"
+
+namespace UE::PixelStreaming2
+{
+	class PIXELSTREAMING2RTC_API FEpicRtcAudioTrack
+	{
+	public:
+		/**
+		 * @return The id of the underlying EpicRtc data track.
+		 */
+		EpicRtcStringView GetTrackId() const;
+
+	protected:
+		TRefCountPtr<EpicRtcAudioTrackInterface> Track;
+	};
+} // namespace UE::PixelStreaming2
\ No newline at end of file
diff --git a/Engine/Plugins/Media/PixelStreaming2/Source/PixelStreaming2RTC/Internal/EpicRtcDataTrack.h b/Engine/Plugins/Media/PixelStreaming2/Source/PixelStreaming2RTC/Internal/EpicRtcDataTrack.h
new file mode 100644
--- /dev/null
+++ b/Engine/Plugins/Media/PixelStreaming2/Source/PixelStreaming2RTC/Internal/EpicRtcDataTrack.h
@@ -0,0 +1,162 @@
+// Copyright Epic Games, Inc. All Rights Reserved.
+
+#pragma once
+
+#include "IPixelStreaming2DataProtocol.h"
+#include "Templates/SharedPointer.h"
+
+#include "epic_rtc/core/data_track.h"
+
+namespace UE::PixelStreaming2
+{
+	template <typename T>
+	inline size_t ValueSize(T&& Value)
+	{
+		return sizeof(Value);
+	}
+
+	template <typename T>
+	inline const void* ValueLoc(T&& Value)
+	{
+		return &Value;
+	}
+
+	inline size_t ValueSize(FString&& Value)
+	{
+		return Value.Len() * sizeof(TCHAR);
+	}
+
+	inline const void* ValueLoc(FString&& Value)
+	{
+		return *Value;
+	}
+
+	struct PIXELSTREAMING2RTC_API FBufferBuilder
+	{
+		TArray<uint8> Buffer;
+		size_t		  Pos;
+
+		FBufferBuilder(size_t size)
+			: Pos(0)
+		{
+			Buffer.SetNum(size);
+		}
+
+		size_t Serialize(const void* Data, size_t DataSize)
+		{
+			check(Pos + DataSize <= Buffer.Num());
+			FMemory::Memcpy(Buffer.GetData() + Pos, Data, DataSize);
+			return Pos + DataSize;
+		}
+
+		template <typename T>
+		void Insert(T&& Value)
+		{
+			const size_t VSize = ValueSize(Forward<T>(Value));
+			const void*	 VLoc = ValueLoc(Forward<T>(Value));
+			check(Pos + VSize <= Buffer.Num());
+			FMemory::Memcpy(Buffer.GetData() + Pos, VLoc, VSize);
+			Pos += VSize;
+		}
+	};
+
+	class PIXELSTREAMING2RTC_API FEpicRtcDataTrack : public TSharedFromThis<FEpicRtcDataTrack>
+	{
+	public:
+		static TSharedPtr<FEpicRtcDataTrack> Create(TRefCountPtr<EpicRtcDataTrackInterface> InTrack, TWeakPtr<IPixelStreaming2DataProtocol> InDataProtocol);
+
+		virtual ~FEpicRtcDataTrack() = default;
+
+		/**
+		 * Sends a series of arguments to the data channel with the given type.
+		 * @param MessageType The name of the message you want to send. This message name must be registered IPixelStreaming2InputHandler::GetFromStreamerProtocol().
+		 * @returns True of the message was successfully sent.
+		 */
+		template <typename... Args>
+		bool SendMessage(FString MessageType, Args... VarArgs) const
+		{
+			if (!IsActive())
+			{
+				return false;
+			}
+
+			uint8 MessageId;
+			if (!GetMessageId(MessageType, MessageId))
+			{
+				return false;
+			}
+
+			FBufferBuilder Builder = EncodeMessage(MessageId, Forward<Args>(VarArgs)...);
+
+			return Send(Builder.Buffer);
+		}
+
+		/**
+		 * Sends a large buffer of data to the data track, will chunk into multiple data frames if frames greater than 16KB.
+		 * @param MessageType The name of the message, it must be registered in IPixelStreaming2InputHandler::GetTo/FromStreamerProtocol()
+		 * @param DataBytes The raw byte buffer to send.
+		 * @returns True of the message was successfully sent.
+		 */
+		bool SendArbitraryData(const FString& MessageType, const TArray64<uint8>& DataBytes) const;
+
+		/**
+		 * @return The id of the underlying EpicRtc data track.
+		 */
+		EpicRtcStringView GetId() const;
+
+		/**
+		 * @return The state of the underlying EpicRtc data track.
+		 */
+		EpicRtcTrackState GetState() const { return Track->GetState(); }
+
+		void SetSendTrack(TRefCountPtr<EpicRtcDataTrackInterface> InSendTrack) { SendTrack = InSendTrack; }
+
+	protected:
+		FEpicRtcDataTrack(TRefCountPtr<EpicRtcDataTrackInterface> InTrack, TWeakPtr<IPixelStreaming2DataProtocol> InDataProtocol);
+		FEpicRtcDataTrack(TSharedPtr<FEpicRtcDataTrack> InTrack, TWeakPtr<IPixelStreaming2DataProtocol> InDataProtocol);
+
+		virtual void PrependData(FBufferBuilder& Builder) const {};
+
+		bool IsActive() const;
+
+		bool GetMessageId(const FString& MessageType, uint8& OutMessageId) const;
+
+	private:
+		TRefCountPtr<EpicRtcDataTrackInterface> Track;
+
+		/**
+		 * Track that is used for sending data with Consumer/Producer architecture.
+		 */
+		TRefCountPtr<EpicRtcDataTrackInterface> SendTrack;
+
+		TWeakPtr<IPixelStreaming2DataProtocol> WeakDataProtocol;
+
+		bool Send(TArray<uint8>& Buffer) const;
+
+		template <typename... Args>
+		FBufferBuilder EncodeMessage(uint8 Type, Args... VarArgs) const
+		{
+
+			FBufferBuilder Builder(sizeof(Type) + (0 + ... + ValueSize(Forward<Args>(VarArgs))));
+			PrependData(Builder);
+			Builder.Insert(Type);
+			(Builder.Insert(Forward<Args>(VarArgs)), ...);
+
+			return MoveTemp(Builder);
+		}
+	};
+
+	class PIXELSTREAMING2RTC_API FEpicRtcMutliplexDataTrack : public FEpicRtcDataTrack
+	{
+	public:
+		static TSharedPtr<FEpicRtcMutliplexDataTrack> Create(TSharedPtr<FEpicRtcDataTrack> InTrack, TWeakPtr<IPixelStreaming2DataProtocol> InDataProtocol, const FString& InPlayerId);
+
+	protected:
+		FEpicRtcMutliplexDataTrack(TSharedPtr<FEpicRtcDataTrack> InTrack, TWeakPtr<IPixelStreaming2DataProtocol> InDataProtocol, const FString& InPlayerId);
+
+		virtual void PrependData(FBufferBuilder& Builder) const override;
+
+	private:
+		const FString PlayerId;
+	};
+} // namespace UE::PixelStreaming2
\ No newline at end of file
diff --git a/Engine/Plugins/Media/PixelStreaming2/Source/PixelStreaming2RTC/Internal/EpicRtcLogging.h b/Engine/Plugins/Media/PixelStreaming2/Source/PixelStreaming2RTC/Internal/EpicRtcLogging.h
--- a/Engine/Plugins/Media/PixelStreaming2/Source/PixelStreaming2RTC/Internal/EpicRtcLogging.h
+++ b/Engine/Plugins/Media/PixelStreaming2/Source/PixelStreaming2RTC/Internal/EpicRtcLogging.h
@@ -6,11 +6,33 @@
 
 #include "epic_rtc/common/logging.h"
 
-DECLARE_LOG_CATEGORY_EXTERN(LogPixelStreaming2EpicRtc, Log, All);
-DECLARE_LOG_CATEGORY_EXTERN(LogPixelStreaming2WebRtc, Fatal, All);
+PIXELSTREAMING2RTC_API DECLARE_LOG_CATEGORY_EXTERN(LogPixelStreaming2EpicRtc, Log, All);
+PIXELSTREAMING2RTC_API DECLARE_LOG_CATEGORY_EXTERN(LogPixelStreaming2WebRtc, Fatal, All);
 
 namespace UE::PixelStreaming2
 {
+	constexpr EpicRtcLogLevel UnrealLogToEpicRtcCategoryMap[] = {
+		EpicRtcLogLevel::Off,
+		EpicRtcLogLevel::Critical,
+		EpicRtcLogLevel::Error,
+		EpicRtcLogLevel::Warning,
+		EpicRtcLogLevel::Info,
+		EpicRtcLogLevel::Info,
+		EpicRtcLogLevel::Debug,
+		EpicRtcLogLevel::Trace,
+		EpicRtcLogLevel::Trace
+	};
+
+	static_assert(EpicRtcLogLevel::Off == UnrealLogToEpicRtcCategoryMap[ELogVerbosity::Type::NoLogging]);
+	static_assert(EpicRtcLogLevel::Critical == UnrealLogToEpicRtcCategoryMap[ELogVerbosity::Type::Fatal]);
+	static_assert(EpicRtcLogLevel::Error == UnrealLogToEpicRtcCategoryMap[ELogVerbosity::Type::Error]);
+	static_assert(EpicRtcLogLevel::Warning == UnrealLogToEpicRtcCategoryMap[ELogVerbosity::Type::Warning]);
+	static_assert(EpicRtcLogLevel::Info == UnrealLogToEpicRtcCategoryMap[ELogVerbosity::Type::Display]);
+	static_assert(EpicRtcLogLevel::Info == UnrealLogToEpicRtcCategoryMap[ELogVerbosity::Type::Log]);
+	static_assert(EpicRtcLogLevel::Debug == UnrealLogToEpicRtcCategoryMap[ELogVerbosity::Type::Verbose]);
+	static_assert(EpicRtcLogLevel::Trace == UnrealLogToEpicRtcCategoryMap[ELogVerbosity::Type::VeryVerbose]);
+	static_assert(EpicRtcLogLevel::Trace == UnrealLogToEpicRtcCategoryMap[ELogVerbosity::Type::All]);
+
 	class PIXELSTREAMING2RTC_API FEpicRtcLogsRedirector : public EpicRtcLoggerInterface
 	{
 	public:
diff --git a/Engine/Plugins/Media/PixelStreaming2/Source/PixelStreaming2RTC/Internal/EpicRtcStatsCollector.h b/Engine/Plugins/Media/PixelStreaming2/Source/PixelStreaming2RTC/Internal/EpicRtcStatsCollector.h
new file mode 100644
--- /dev/null
+++ b/Engine/Plugins/Media/PixelStreaming2/Source/PixelStreaming2RTC/Internal/EpicRtcStatsCollector.h
@@ -0,0 +1,32 @@
+// Copyright Epic Games, Inc. All Rights Reserved.
+
+#pragma once
+
+#include "Delegates/Delegate.h"
+#include "Templates/RefCounting.h"
+#include "Templates/SharedPointer.h"
+
+#include "epic_rtc/core/stats.h"
+#include "epic_rtc_helper/memory/ref_count_impl_helper.h"
+
+namespace UE::PixelStreaming2
+{
+	class PIXELSTREAMING2RTC_API FEpicRtcStatsCollector : public EpicRtcStatsCollectorCallbackInterface
+	{
+	public:
+		FEpicRtcStatsCollector() = default;
+		~FEpicRtcStatsCollector() = default;
+
+		// Begin EpicRtcStatsCollectorCallbackInterface interface
+		void			 OnStatsDelivered(const EpicRtcStatsReport& InReport) override;
+		// End EpicRtcStatsCollectorCallbackInterface interface
+
+		// Begin EpicRtcRefCountInterface
+		EPICRTC_REFCOUNT_INTERFACE_IN_PLACE
+		// End EpicRtcStatsCollectorCallbackInterface interface
+	public:
+
+		DECLARE_EVENT_TwoParams(FEpicRtcStatsCollector, FOnStatsReady, const FString&, const EpicRtcConnectionStats&);
+		FOnStatsReady OnStatsReady;
+	};
+} // namespace UE::PixelStreaming2
\ No newline at end of file
diff --git a/Engine/Plugins/Media/PixelStreaming2/Source/PixelStreaming2RTC/Internal/EpicRtcVideoCapturer.h b/Engine/Plugins/Media/PixelStreaming2/Source/PixelStreaming2RTC/Internal/EpicRtcVideoCapturer.h
new file mode 100644
--- /dev/null
+++ b/Engine/Plugins/Media/PixelStreaming2/Source/PixelStreaming2RTC/Internal/EpicRtcVideoCapturer.h
@@ -0,0 +1,22 @@
+// Copyright Epic Games, Inc. All Rights Reserved.
+
+#pragma once
+
+#include "VideoCapturer.h"
+
+#include "epic_rtc/core/video/video_buffer.h"
+
+namespace UE::PixelStreaming2
+{
+	class PIXELSTREAMING2RTC_API FEpicRtcVideoCapturer : public FVideoCapturer
+	{
+	public:
+		static TSharedPtr<FEpicRtcVideoCapturer> Create(TSharedPtr<FVideoProducer> VideoProducer = nullptr);
+		virtual ~FEpicRtcVideoCapturer() = default;
+
+		TRefCountPtr<EpicRtcVideoBufferInterface> GetFrameBuffer();
+
+	private:
+		FEpicRtcVideoCapturer(TSharedPtr<FVideoProducer> VideoProducer);
+	};
+} // namespace UE::PixelStreaming2
diff --git a/Engine/Plugins/Media/PixelStreaming2/Source/PixelStreaming2RTC/Internal/EpicRtcVideoSink.h b/Engine/Plugins/Media/PixelStreaming2/Source/PixelStreaming2RTC/Internal/EpicRtcVideoSink.h
new file mode 100644
--- /dev/null
+++ b/Engine/Plugins/Media/PixelStreaming2/Source/PixelStreaming2RTC/Internal/EpicRtcVideoSink.h
@@ -0,0 +1,35 @@
+// Copyright Epic Games, Inc. All Rights Reserved.
+
+#pragma once
+
+#include "Containers/Set.h"
+#include "EpicRtcVideoTrack.h"
+#include "VideoSink.h"
+#include "RendererInterface.h"
+
+#include "epic_rtc/core/video/video_frame.h"
+
+namespace UE::PixelStreaming2
+{
+	/**
+	 * Video sink class that receives a frame from EpicRtc and passes the frame to all added consumers
+	 */
+	class PIXELSTREAMING2RTC_API FEpicRtcVideoSink : public FVideoSink, public FEpicRtcVideoTrack
+	{
+	public:
+		static TSharedPtr<FEpicRtcVideoSink> Create(TRefCountPtr<EpicRtcVideoTrackInterface> InTrack);
+		// Note: destructor will call destroy on any attached video consumers
+		virtual ~FEpicRtcVideoSink() = default;
+
+		void OnVideoData(const EpicRtcVideoFrame& Frame);
+
+	private:
+		FEpicRtcVideoSink(TRefCountPtr<EpicRtcVideoTrackInterface> InTrack);
+
+		FCriticalSection				  RenderSyncContext;
+		FPooledRenderTargetDesc			  RenderTargetDescriptor;
+		TRefCountPtr<IPooledRenderTarget> RenderTarget;
+		TArray<uint8_t>					  Buffer;
+		FTextureRHIRef					  SourceTexture;
+	};
+} // namespace UE::PixelStreaming2
\ No newline at end of file
diff --git a/Engine/Plugins/Media/PixelStreaming2/Source/PixelStreaming2RTC/Internal/EpicRtcVideoSource.h b/Engine/Plugins/Media/PixelStreaming2/Source/PixelStreaming2RTC/Internal/EpicRtcVideoSource.h
new file mode 100644
--- /dev/null
+++ b/Engine/Plugins/Media/PixelStreaming2/Source/PixelStreaming2RTC/Internal/EpicRtcVideoSource.h
@@ -0,0 +1,40 @@
+// Copyright Epic Games, Inc. All Rights Reserved.
+
+#pragma once
+
+#include "EpicRtcVideoCapturer.h"
+#include "EpicRtcVideoTrack.h"
+#include "VideoSource.h"
+
+namespace UE::PixelStreaming2
+{
+	class FVideoSourceGroup;
+	/**
+	 * A source of video frames for an EpicRtc peer. Has a video input which will provide
+	 * frame data to this source. The source will then pass that data to an adapter
+	 * which will have one or many adapt processes that are provided by the input object
+	 * and are responsible for converting the frame data to the format required for
+	 * the selected video encoder.
+	 * This video source should be contained within a FVideoSourceGroup which is responsible
+	 * for telling each source to push a frame to EpicRtc at the expected rate. This source
+	 * will make sure that the adapter has valid output and if so will create a frame
+	 * for EpicRtc. Otherwise it will continue to wait until the next frame.
+	 */
+	class PIXELSTREAMING2RTC_API FEpicRtcVideoSource : public FVideoSource, public FEpicRtcVideoTrack
+	{
+	public:
+		static TSharedPtr<FEpicRtcVideoSource> Create(TRefCountPtr<EpicRtcVideoTrackInterface> InVideoTrack, TSharedPtr<FEpicRtcVideoCapturer> InVideoCapturer, TSharedPtr<FVideoSourceGroup> InVideoSourceGroup);
+		virtual ~FEpicRtcVideoSource() = default;
+
+		virtual void PushFrame() override;
+		virtual void ForceKeyFrame() override;
+
+		TRefCountPtr<EpicRtcVideoTrackInterface> GetVideoTrack();
+
+	private:
+		FEpicRtcVideoSource(TRefCountPtr<EpicRtcVideoTrackInterface> InVideoTrack, TSharedPtr<FEpicRtcVideoCapturer> InVideoCapturer);
+
+	private:
+		TSharedPtr<FEpicRtcVideoCapturer> VideoCapturer;
+	};
+} // namespace UE::PixelStreaming2
diff --git a/Engine/Plugins/Media/PixelStreaming2/Source/PixelStreaming2RTC/Internal/EpicRtcVideoTrack.h b/Engine/Plugins/Media/PixelStreaming2/Source/PixelStreaming2RTC/Internal/EpicRtcVideoTrack.h
new file mode 100644
--- /dev/null
+++ b/Engine/Plugins/Media/PixelStreaming2/Source/PixelStreaming2RTC/Internal/EpicRtcVideoTrack.h
@@ -0,0 +1,22 @@
+// Copyright Epic Games, Inc. All Rights Reserved.
+
+#pragma once
+
+#include "Templates/RefCounting.h"
+
+#include "epic_rtc/core/video/video_track.h"
+
+namespace UE::PixelStreaming2
+{
+	class PIXELSTREAMING2RTC_API FEpicRtcVideoTrack
+	{
+	public:
+		/**
+		 * @return The id of the underlying EpicRtc data track.
+		 */
+		EpicRtcStringView GetTrackId() const;
+
+	protected:
+		TRefCountPtr<EpicRtcVideoTrackInterface> Track;
+	};
+} // namespace UE::PixelStreaming2
\ No newline at end of file
diff --git a/Engine/Plugins/Media/PixelStreaming2/Source/PixelStreaming2RTC/Internal/UtilsString.h b/Engine/Plugins/Media/PixelStreaming2/Source/PixelStreaming2RTC/Internal/UtilsString.h
new file mode 100644
--- /dev/null
+++ b/Engine/Plugins/Media/PixelStreaming2/Source/PixelStreaming2RTC/Internal/UtilsString.h
@@ -0,0 +1,191 @@
+// Copyright Epic Games, Inc. All Rights Reserved.
+
+#pragma once
+
+#include <string>
+
+#include "Containers/StringConv.h"
+#include "Containers/UnrealString.h"
+#include "Containers/Utf8String.h"
+#include "Dom/JsonObject.h"
+#include "Policies/PrettyJsonPrintPolicy.h"
+#include "Policies/CondensedJsonPrintPolicy.h"
+#include "Serialization/JsonSerializer.h"
+
+#include "epic_rtc/common/common.h"
+#include "epic_rtc/containers/epic_rtc_string_view.h"
+#include "epic_rtc/core/video/video_common.h"
+
+namespace UE::PixelStreaming2
+{
+	inline FString ToString(EpicRtcErrorCode Error)
+	{
+		switch (Error)
+		{
+			case EpicRtcErrorCode::Ok:
+			{
+				return TEXT("Ok");
+			}
+			case EpicRtcErrorCode::GeneralError:
+			{
+				return TEXT("GeneralError");
+			}
+			case EpicRtcErrorCode::BadState:
+			{
+				return TEXT("BadState");
+			}
+			case EpicRtcErrorCode::Timeout:
+			{
+				return TEXT("Timeout");
+			}
+			case EpicRtcErrorCode::Unsupported:
+			{
+				return TEXT("Unsupported");
+			}
+			case EpicRtcErrorCode::PlatformError:
+			{
+				return TEXT("PlatformError");
+			}
+			case EpicRtcErrorCode::FoundExistingPlatform:
+			{
+				return TEXT("FoundExistingPlatform");
+			}
+			case EpicRtcErrorCode::ConferenceAlreadyExists:
+			{
+				return TEXT("ConferenceAlreadyExists");
+			}
+			case EpicRtcErrorCode::ConferenceDoesNotExists:
+			{
+				return TEXT("ConferenceDoesNotExists");
+			}
+			case EpicRtcErrorCode::ImATeapot:
+			{
+				return TEXT("ImATeapot");
+			}
+			case EpicRtcErrorCode::ConferenceError:
+			{
+				return TEXT("ConferenceError");
+			}
+			case EpicRtcErrorCode::SessionAlreadyExists:
+			{
+				return TEXT("SessionAlreadyExists");
+			}
+			case EpicRtcErrorCode::SessionDoesNotExist:
+			{
+				return TEXT("SessionDoesNotExist");
+			}
+			case EpicRtcErrorCode::SessionError:
+			{
+				return TEXT("SessionError");
+			}
+			case EpicRtcErrorCode::SessionCannotConnect:
+			{
+				return TEXT("SessionCannotConnect");
+			}
+			case EpicRtcErrorCode::SessionDisconnected:
+			{
+				return TEXT("SessionDisconnected");
+			}
+			case EpicRtcErrorCode::SessionCannotCreateRoom:
+			{
+				return TEXT("SessionCannotCreateRoom");
+			}
+			default:
+			{
+				return TEXT("Unknown");
+			}
+		}
+	}
+
+	inline FString ToString(EpicRtcVideoCodec Codec)
+	{
+		switch (Codec)
+		{
+			case EpicRtcVideoCodec::AV1:
+			{
+				return TEXT("AV1");
+			}
+			case EpicRtcVideoCodec::H264:
+			{
+				return TEXT("H264");
+			}
+			case EpicRtcVideoCodec::VP8:
+			{
+				return TEXT("VP8");
+			}
+			case EpicRtcVideoCodec::VP9:
+			{
+				return TEXT("VP9");
+			}
+			default:
+			{
+				return TEXT("Unknown");
+			}
+		}
+	}
+
+	inline FString ToString(EpicRtcRoomState State)
+	{
+		switch (State)
+		{
+			case EpicRtcRoomState::New:
+				return TEXT("New");
+			case EpicRtcRoomState::Pending:
+				return TEXT("Pending");
+			case EpicRtcRoomState::Joined:
+				return TEXT("Joined");
+			case EpicRtcRoomState::Left:
+				return TEXT("Left");
+			case EpicRtcRoomState::Failed:
+				return TEXT("Failed");
+			case EpicRtcRoomState::Exiting:
+				return TEXT("Exiting");
+		}
+		return TEXT("Unknown");
+	}
+
+	inline FString ToString(const TSharedPtr<FJsonObject>& JsonObj, bool bPretty = true)
+	{
+		FString Res;
+		if (bPretty)
+		{
+			auto JsonWriter = TJsonWriterFactory<TCHAR, TPrettyJsonPrintPolicy<TCHAR>>::Create(&Res);
+			FJsonSerializer::Serialize(JsonObj.ToSharedRef(), JsonWriter);
+		}
+		else
+		{
+			auto JsonWriter = TJsonWriterFactory<TCHAR, TCondensedJsonPrintPolicy<TCHAR>>::Create(&Res);
+			FJsonSerializer::Serialize(JsonObj.ToSharedRef(), JsonWriter);
+		}
+		return Res;
+	}
+
+	inline FString ToString(const EpicRtcStringView& Str)
+	{
+		FUtf8String Utf8String = FUtf8String::ConstructFromPtrSize(Str._ptr, Str._length);
+		return FString(Utf8String);
+	}
+
+	inline EpicRtcStringView ToEpicRtcStringView(const FUtf8String& Str)
+	{
+		return EpicRtcStringView{ ._ptr = (const char*)*Str, ._length = static_cast<uint64>(Str.Len()) };
+	}
+
+	/**
+	 * Reads a string represented by 2 bytes length (in bytes) followed by UTF16 characters.
+	 * String and length are encoded in little endian format.
+	 */
+	inline FString ReadString(const uint8*& Data, uint32_t& Size)
+	{
+		uint16_t BytesLength = Data[1] << 8 | Data[0];
+		check(Size >= (uint32_t)(BytesLength + 2));
+		Data += 2;
+		Size -= 2;
+
+		FString Message(BytesLength / sizeof(TCHAR), reinterpret_cast<const TCHAR*>(Data));
+		Data += BytesLength;
+		Size -= BytesLength;
+
+		return Message;
+	}
+} // namespace UE::PixelStreaming2
diff --git a/Engine/Plugins/Media/PixelStreaming2/Source/PixelStreaming2RTC/Private/EpicRtcAudioCapturer.h b/Engine/Plugins/Media/PixelStreaming2/Source/PixelStreaming2RTC/Private/EpicRtcAudioCapturer.h
deleted file mode 100644
diff --git a/Engine/Plugins/Media/PixelStreaming2/Source/PixelStreaming2RTC/Private/EpicRtcAudioSink.h b/Engine/Plugins/Media/PixelStreaming2/Source/PixelStreaming2RTC/Private/EpicRtcAudioSink.h
deleted file mode 100644
diff --git a/Engine/Plugins/Media/PixelStreaming2/Source/PixelStreaming2RTC/Private/EpicRtcAudioSource.cpp b/Engine/Plugins/Media/PixelStreaming2/Source/PixelStreaming2RTC/Private/EpicRtcAudioSource.cpp
--- a/Engine/Plugins/Media/PixelStreaming2/Source/PixelStreaming2RTC/Private/EpicRtcAudioSource.cpp
+++ b/Engine/Plugins/Media/PixelStreaming2/Source/PixelStreaming2RTC/Private/EpicRtcAudioSource.cpp
@@ -8,12 +8,11 @@
 namespace UE::PixelStreaming2
 {
 
-	TSharedPtr<FEpicRtcAudioSource> FEpicRtcAudioSource::Create(TRefCountPtr<EpicRtcAudioTrackInterface> InTrack)
+	TSharedPtr<FEpicRtcAudioSource> FEpicRtcAudioSource::Create(TRefCountPtr<EpicRtcAudioTrackInterface> InTrack, TSharedPtr<FEpicRtcAudioCapturer> InCapturer)
 	{
 		TSharedPtr<FEpicRtcAudioSource> AudioTrack = TSharedPtr<FEpicRtcAudioSource>(new FEpicRtcAudioSource(InTrack));
 
-		TSharedPtr<FEpicRtcAudioCapturer> Capturer = FPixelStreaming2RTCModule::GetModule()->GetAudioCapturer();
-		Capturer->OnAudioBuffer.AddSP(AudioTrack.ToSharedRef(), &FEpicRtcAudioSource::OnAudioBuffer);
+		InCapturer->OnAudioBuffer.AddSP(AudioTrack.ToSharedRef(), &FEpicRtcAudioSource::OnAudioBuffer);
 
 		return AudioTrack;
 	}
diff --git a/Engine/Plugins/Media/PixelStreaming2/Source/PixelStreaming2RTC/Private/EpicRtcAudioSource.h b/Engine/Plugins/Media/PixelStreaming2/Source/PixelStreaming2RTC/Private/EpicRtcAudioSource.h
deleted file mode 100644
diff --git a/Engine/Plugins/Media/PixelStreaming2/Source/PixelStreaming2RTC/Private/EpicRtcAudioTrack.h b/Engine/Plugins/Media/PixelStreaming2/Source/PixelStreaming2RTC/Private/EpicRtcAudioTrack.h
deleted file mode 100644
diff --git a/Engine/Plugins/Media/PixelStreaming2/Source/PixelStreaming2RTC/Private/EpicRtcDataTrack.cpp b/Engine/Plugins/Media/PixelStreaming2/Source/PixelStreaming2RTC/Private/EpicRtcDataTrack.cpp
--- a/Engine/Plugins/Media/PixelStreaming2/Source/PixelStreaming2RTC/Private/EpicRtcDataTrack.cpp
+++ b/Engine/Plugins/Media/PixelStreaming2/Source/PixelStreaming2RTC/Private/EpicRtcDataTrack.cpp
@@ -2,6 +2,7 @@
 
 #include "EpicRtcDataTrack.h"
 
+#include "Logging.h"
 #include "UtilsCommon.h"
 #include "UtilsString.h"
 
diff --git a/Engine/Plugins/Media/PixelStreaming2/Source/PixelStreaming2RTC/Private/EpicRtcDataTrack.h b/Engine/Plugins/Media/PixelStreaming2/Source/PixelStreaming2RTC/Private/EpicRtcDataTrack.h
deleted file mode 100644
diff --git a/Engine/Plugins/Media/PixelStreaming2/Source/PixelStreaming2RTC/Private/EpicRtcStatsCollector.h b/Engine/Plugins/Media/PixelStreaming2/Source/PixelStreaming2RTC/Private/EpicRtcStatsCollector.h
deleted file mode 100644
diff --git a/Engine/Plugins/Media/PixelStreaming2/Source/PixelStreaming2RTC/Private/EpicRtcStreamer.cpp b/Engine/Plugins/Media/PixelStreaming2/Source/PixelStreaming2RTC/Private/EpicRtcStreamer.cpp
--- a/Engine/Plugins/Media/PixelStreaming2/Source/PixelStreaming2RTC/Private/EpicRtcStreamer.cpp
+++ b/Engine/Plugins/Media/PixelStreaming2/Source/PixelStreaming2RTC/Private/EpicRtcStreamer.cpp
@@ -1183,7 +1183,7 @@ namespace UE::PixelStreaming2
 			}
 			else
 			{
-				PlayerContext->AudioSource = FEpicRtcAudioSource::Create(AudioTrack);
+				PlayerContext->AudioSource = FEpicRtcAudioSource::Create(AudioTrack, FPixelStreaming2RTCModule::GetModule()->GetAudioCapturer());
 			}
 
 			if (UPixelStreaming2Delegates* Delegates = UPixelStreaming2Delegates::Get())
diff --git a/Engine/Plugins/Media/PixelStreaming2/Source/PixelStreaming2RTC/Private/EpicRtcVideoCapturer.h b/Engine/Plugins/Media/PixelStreaming2/Source/PixelStreaming2RTC/Private/EpicRtcVideoCapturer.h
deleted file mode 100644
diff --git a/Engine/Plugins/Media/PixelStreaming2/Source/PixelStreaming2RTC/Private/EpicRtcVideoSink.h b/Engine/Plugins/Media/PixelStreaming2/Source/PixelStreaming2RTC/Private/EpicRtcVideoSink.h
deleted file mode 100644
diff --git a/Engine/Plugins/Media/PixelStreaming2/Source/PixelStreaming2RTC/Private/EpicRtcVideoSource.h b/Engine/Plugins/Media/PixelStreaming2/Source/PixelStreaming2RTC/Private/EpicRtcVideoSource.h
deleted file mode 100644
diff --git a/Engine/Plugins/Media/PixelStreaming2/Source/PixelStreaming2RTC/Private/EpicRtcVideoTrack.h b/Engine/Plugins/Media/PixelStreaming2/Source/PixelStreaming2RTC/Private/EpicRtcVideoTrack.h
deleted file mode 100644
diff --git a/Engine/Plugins/Media/PixelStreaming2/Source/PixelStreaming2RTC/Private/PixelStreaming2RTCModule.cpp b/Engine/Plugins/Media/PixelStreaming2/Source/PixelStreaming2RTC/Private/PixelStreaming2RTCModule.cpp
--- a/Engine/Plugins/Media/PixelStreaming2/Source/PixelStreaming2RTC/Private/PixelStreaming2RTCModule.cpp
+++ b/Engine/Plugins/Media/PixelStreaming2/Source/PixelStreaming2RTC/Private/PixelStreaming2RTCModule.cpp
@@ -69,28 +69,6 @@ THIRD_PARTY_INCLUDES_END
 
 namespace UE::PixelStreaming2
 {
-	constexpr EpicRtcLogLevel UnrealLogToEpicRtcCategoryMap[] = {
-		EpicRtcLogLevel::Off,
-		EpicRtcLogLevel::Critical,
-		EpicRtcLogLevel::Error,
-		EpicRtcLogLevel::Warning,
-		EpicRtcLogLevel::Info,
-		EpicRtcLogLevel::Info,
-		EpicRtcLogLevel::Debug,
-		EpicRtcLogLevel::Trace,
-		EpicRtcLogLevel::Trace
-	};
-
-	static_assert(EpicRtcLogLevel::Off == UnrealLogToEpicRtcCategoryMap[ELogVerbosity::Type::NoLogging]);
-	static_assert(EpicRtcLogLevel::Critical == UnrealLogToEpicRtcCategoryMap[ELogVerbosity::Type::Fatal]);
-	static_assert(EpicRtcLogLevel::Error == UnrealLogToEpicRtcCategoryMap[ELogVerbosity::Type::Error]);
-	static_assert(EpicRtcLogLevel::Warning == UnrealLogToEpicRtcCategoryMap[ELogVerbosity::Type::Warning]);
-	static_assert(EpicRtcLogLevel::Info == UnrealLogToEpicRtcCategoryMap[ELogVerbosity::Type::Display]);
-	static_assert(EpicRtcLogLevel::Info == UnrealLogToEpicRtcCategoryMap[ELogVerbosity::Type::Log]);
-	static_assert(EpicRtcLogLevel::Debug == UnrealLogToEpicRtcCategoryMap[ELogVerbosity::Type::Verbose]);
-	static_assert(EpicRtcLogLevel::Trace == UnrealLogToEpicRtcCategoryMap[ELogVerbosity::Type::VeryVerbose]);
-	static_assert(EpicRtcLogLevel::Trace == UnrealLogToEpicRtcCategoryMap[ELogVerbosity::Type::All]);
-
 	FPixelStreaming2RTCModule* FPixelStreaming2RTCModule::PixelStreaming2Module = nullptr;
 
 	FUtf8String FPixelStreaming2RTCModule::EpicRtcConferenceName("pixel_streaming_conference_instance");
@@ -307,9 +285,6 @@ namespace UE::PixelStreaming2
 
 	bool FPixelStreaming2RTCModule::InitializeEpicRtc()
 	{
-		// Initialize the audio capturer (if not already called by some user code)
-		AudioMixingCapturer = GetAudioCapturer();
-
 		EpicRtcVideoEncoderInitializers = { new FEpicRtcVideoEncoderInitializer() };
 		EpicRtcVideoDecoderInitializers = { new FEpicRtcVideoDecoderInitializer() };
 
diff --git a/Engine/Plugins/Media/PixelStreaming2/Source/PixelStreaming2RTC/Private/StreamerReconnectTimer.cpp b/Engine/Plugins/Media/PixelStreaming2/Source/PixelStreaming2RTC/Private/StreamerReconnectTimer.cpp
deleted file mode 100644
diff --git a/Engine/Plugins/Media/PixelStreaming2/Source/PixelStreaming2RTC/Private/StreamerReconnectTimer.h b/Engine/Plugins/Media/PixelStreaming2/Source/PixelStreaming2RTC/Private/StreamerReconnectTimer.h
deleted file mode 100644
diff --git a/Engine/Plugins/Media/PixelStreaming2/Source/PixelStreaming2RTC/Private/UtilsString.h b/Engine/Plugins/Media/PixelStreaming2/Source/PixelStreaming2RTC/Private/UtilsString.h
deleted file mode 100644
