[PS2] Fix: Non-unity compilation errors

diff --git a/Engine/Plugins/Media/PixelStreaming2/Source/PixelStreaming2Editor/Internal/VideoProducerViewport.h b/Engine/Plugins/Media/PixelStreaming2/Source/PixelStreaming2Editor/Internal/VideoProducerViewport.h
--- a/Engine/Plugins/Media/PixelStreaming2/Source/PixelStreaming2Editor/Internal/VideoProducerViewport.h
+++ b/Engine/Plugins/Media/PixelStreaming2/Source/PixelStreaming2Editor/Internal/VideoProducerViewport.h
@@ -2,10 +2,11 @@
 
 #pragma once
 
-#include "VideoProducer.h"
-#include "IPixelStreaming2Streamer.h"
 #include "Delegates/IDelegateInstance.h"
+#include "IPixelStreaming2Streamer.h"
 #include "UnrealClient.h"
+#include "VideoProducer.h"
+#include "Widgets/SWindow.h"
 
 namespace UE::EditorPixelStreaming2
 {
diff --git a/Engine/Plugins/Media/PixelStreaming2/Source/PixelStreaming2Editor/PixelStreaming2Editor.build.cs b/Engine/Plugins/Media/PixelStreaming2/Source/PixelStreaming2Editor/PixelStreaming2Editor.build.cs
--- a/Engine/Plugins/Media/PixelStreaming2/Source/PixelStreaming2Editor/PixelStreaming2Editor.build.cs
+++ b/Engine/Plugins/Media/PixelStreaming2/Source/PixelStreaming2Editor/PixelStreaming2Editor.build.cs
@@ -22,7 +22,6 @@ namespace UnrealBuildTool.Rules
 				"PixelStreaming2RTC",
 				"PixelStreaming2Settings",
 				"Slate",
-				"SlateCore",
 				"EngineSettings",
 				"InputCore",
 				"Json",
@@ -37,6 +36,11 @@ namespace UnrealBuildTool.Rules
 				"PixelStreaming2Core"
 			});
 
+			PublicDependencyModuleNames.AddRange(new string[] 
+			{
+				"SlateCore"
+			});
+
 			if (Target.bBuildEditor)
 			{
 				PrivateDependencyModuleNames.AddRange(new string[]
diff --git a/Engine/Plugins/Media/PixelStreaming2/Source/PixelStreaming2Editor/Private/PixelStreaming2EditorModule.h b/Engine/Plugins/Media/PixelStreaming2/Source/PixelStreaming2Editor/Private/PixelStreaming2EditorModule.h
--- a/Engine/Plugins/Media/PixelStreaming2/Source/PixelStreaming2Editor/Private/PixelStreaming2EditorModule.h
+++ b/Engine/Plugins/Media/PixelStreaming2/Source/PixelStreaming2Editor/Private/PixelStreaming2EditorModule.h
@@ -5,6 +5,7 @@
 #include "IPixelStreaming2EditorModule.h"
 #include "IPixelStreaming2RTCModule.h"
 #include "PixelStreaming2Servers.h"
+#include "Widgets/SWindow.h"
 
 namespace UE::EditorPixelStreaming2
 {
diff --git a/Engine/Plugins/Media/PixelStreaming2/Source/PixelStreaming2Editor/Private/VideoProducerViewport.cpp b/Engine/Plugins/Media/PixelStreaming2/Source/PixelStreaming2Editor/Private/VideoProducerViewport.cpp
--- a/Engine/Plugins/Media/PixelStreaming2/Source/PixelStreaming2Editor/Private/VideoProducerViewport.cpp
+++ b/Engine/Plugins/Media/PixelStreaming2/Source/PixelStreaming2Editor/Private/VideoProducerViewport.cpp
@@ -5,6 +5,7 @@
 #include "Async/Async.h"
 #include "Engine/GameViewportClient.h"
 #include "Framework/Application/SlateApplication.h"
+#include "IPixelStreaming2InputHandler.h"
 #include "LevelEditor.h"
 #include "LevelEditorViewport.h"
 #include "Modules/ModuleManager.h"
diff --git a/Engine/Plugins/Media/PixelStreaming2/Source/PixelStreaming2RTC/Private/Tests/MediaDirectionTest.cpp b/Engine/Plugins/Media/PixelStreaming2/Source/PixelStreaming2RTC/Private/Tests/MediaDirectionTest.cpp
--- a/Engine/Plugins/Media/PixelStreaming2/Source/PixelStreaming2RTC/Private/Tests/MediaDirectionTest.cpp
+++ b/Engine/Plugins/Media/PixelStreaming2/Source/PixelStreaming2RTC/Private/Tests/MediaDirectionTest.cpp
@@ -2,6 +2,7 @@
 
 #if WITH_DEV_AUTOMATION_TESTS
 
+	#include "Logging.h"
 	#include "Misc/AutomationTest.h"
 	#include "PixelStreaming2PluginSettings.h"
 	#include "TestUtils.h"
diff --git a/Engine/Plugins/Media/PixelStreaming2/Source/PixelStreaming2RTC/Private/Tests/MultipleSignallingConnectionsTest.cpp b/Engine/Plugins/Media/PixelStreaming2/Source/PixelStreaming2RTC/Private/Tests/MultipleSignallingConnectionsTest.cpp
--- a/Engine/Plugins/Media/PixelStreaming2/Source/PixelStreaming2RTC/Private/Tests/MultipleSignallingConnectionsTest.cpp
+++ b/Engine/Plugins/Media/PixelStreaming2/Source/PixelStreaming2RTC/Private/Tests/MultipleSignallingConnectionsTest.cpp
@@ -1,12 +1,13 @@
 // Copyright Epic Games, Inc. All Rights Reserved.
 
-#include "Misc/AutomationTest.h"
-#include "Tests/AutomationCommon.h"
-#include "PixelStreaming2Servers.h"
-#include "Misc/Paths.h"
 #include "IPixelStreaming2Module.h"
 #include "IPixelStreaming2Streamer.h"
+#include "Logging.h"
+#include "Misc/AutomationTest.h"
+#include "Misc/Paths.h"
+#include "PixelStreaming2Servers.h"
 #include "TestUtils.h"
+#include "Tests/AutomationCommon.h"
 
 #if WITH_DEV_AUTOMATION_TESTS
 
