[PS2] QOL: Refactor InputHandler logic. It is now up to the individual modules to implement decoding of message structures. This enables the use of the base FPixelStreaming2InputHandler with differing protocols and message layouts.

diff --git a/Engine/Plugins/Media/PixelStreaming2/Source/PixelStreaming2Input/Internal/InputHandler.h b/Engine/Plugins/Media/PixelStreaming2/Source/PixelStreaming2Input/Internal/InputHandler.h
new file mode 100644
--- /dev/null
+++ b/Engine/Plugins/Media/PixelStreaming2/Source/PixelStreaming2Input/Internal/InputHandler.h
@@ -0,0 +1,222 @@
+// Copyright Epic Games, Inc. All Rights Reserved.
+
+#pragma once
+
+#include "GenericPlatform/GenericApplicationMessageHandler.h"
+#include "Slate/SceneViewport.h"
+#include "Widgets/SWindow.h"
+#include "InputCoreTypes.h"
+#include "IPixelStreaming2InputHandler.h"
+#include "XRMotionControllerBase.h"
+#include "HeadMountedDisplayTypes.h"
+#include "PixelStreaming2HMDEnums.h"
+
+namespace UE::PixelStreaming2Input
+{
+	class PIXELSTREAMING2INPUT_API FPixelStreaming2InputHandler : public IPixelStreaming2InputHandler, public FXRMotionControllerBase
+	{
+	public:
+		FPixelStreaming2InputHandler();
+		virtual ~FPixelStreaming2InputHandler();
+
+		virtual void Tick(float DeltaTime) override;
+
+		// Poll for controller state and send events if needed
+		virtual void SendControllerEvents() override {};
+
+		// Set which MessageHandler will route input
+		virtual void SetMessageHandler(const TSharedRef<FGenericApplicationMessageHandler>& InTargetHandler) override;
+
+		// Register a custom function to execute when command JSON is received.
+		virtual void SetCommandHandler(const FString& CommandName, const CommandHandlerFn& Handler) override;
+
+		virtual void SetElevatedCheck(const TFunction<bool(FString)>& CheckFn) override;
+		virtual bool IsElevated(const FString& Id) override;
+
+		virtual TSharedPtr<IPixelStreaming2DataProtocol> GetToStreamerProtocol() override;
+		virtual TSharedPtr<IPixelStreaming2DataProtocol> GetFromStreamerProtocol() override;
+
+		// Exec handler to allow console commands to be passed through for debugging
+		virtual bool Exec(UWorld* InWorld, const TCHAR* Cmd, FOutputDevice& Ar) override;
+
+		/**
+		 * IInputInterface pass through functions
+		 */
+		virtual void				SetChannelValue(int32 ControllerId, FForceFeedbackChannelType ChannelType, float Value) override;
+		virtual void				SetChannelValues(int32 ControllerId, const FForceFeedbackValues& values) override;
+		virtual void				OnMessage(FString SourceId, TArray<uint8> Buffer) override;
+		virtual void				SetTargetWindow(TWeakPtr<SWindow> InWindow) override;
+		virtual TWeakPtr<SWindow>	GetTargetWindow() override;
+		virtual void				SetTargetViewport(TWeakPtr<SViewport> InViewport) override;
+		virtual TWeakPtr<SViewport> GetTargetViewport() override;
+		virtual void				SetTargetScreenRect(TWeakPtr<FIntRect> InScreenRect) override;
+		virtual TWeakPtr<FIntRect>	GetTargetScreenRect() override;
+		virtual bool				IsFakingTouchEvents() const override { return bFakingTouchEvents; }
+		virtual void				RegisterMessageHandler(const FString& MessageType, const MessageHandlerFn& Handler) override;
+		virtual MessageHandlerFn	FindMessageHandler(const FString& MessageType) override;
+		virtual void				SetInputType(EPixelStreaming2InputType InInputType) override { InputType = InInputType; };
+		// IMotionController Interface
+		virtual FName			GetMotionControllerDeviceTypeName() const override;
+		virtual bool			GetControllerOrientationAndPosition(const int32 ControllerIndex, const FName MotionSource, FRotator& OutOrientation, FVector& OutPosition, float WorldToMetersScale) const override;
+		virtual ETrackingStatus GetControllerTrackingStatus(const int32 ControllerIndex, const FName MotionSource) const override;
+		virtual void			EnumerateSources(TArray<FMotionControllerSource>& SourcesOut) const override;
+		// End IMotionController Interface
+
+		virtual bool OnKeyChar(TCHAR Character) override;
+		virtual bool OnKeyDown(FKey Key, bool bIsRepeat) override;
+		virtual bool OnKeyUp(FKey Key) override;
+
+		virtual bool OnMouseEnter() override;
+		virtual bool OnMouseLeave() override;
+		virtual bool OnMouseDown(EMouseButtons::Type Button, FIntPoint ScreenPosition) override;
+		virtual bool OnMouseUp(EMouseButtons::Type Button) override;
+		virtual bool OnMouseMove(FIntPoint ScreenLocation, FIntPoint Delta) override;
+		virtual bool OnMouseWheel(FIntPoint ScreenLocation, float MouseWheelDelta) override;
+		virtual bool OnMouseDoubleClick(EMouseButtons::Type Button, FIntPoint ScreenPosition) override;
+
+		virtual bool OnTouchStarted(FIntPoint TouchLocation, int32 TouchIndex, float Force) override;
+		virtual bool OnTouchMoved(FIntPoint TouchLocation, int32 TouchIndex, float Force) override;
+		virtual bool OnTouchEnded(FIntPoint TouchLocation, int32 TouchIndex) override;
+
+		virtual uint8 OnControllerConnected() override;
+		virtual bool  OnControllerAnalog(uint8 ControllerIndex, FKey Key, double AxisValue) override;
+		virtual bool  OnControllerButtonPressed(uint8 ControllerIndex, FKey Key, bool bIsRepeat) override;
+		virtual bool  OnControllerButtonReleased(uint8 ControllerIndex, FKey Key) override;
+		virtual bool  OnControllerDisconnected(uint8 ControllerIndex) override;
+
+		virtual bool OnXREyeViews(FTransform LeftEyeTransform, FMatrix LeftEyeProjectionMatrix, FTransform RightEyeTransform, FMatrix RightEyeProjectionMatrix, FTransform HMDTransform) override;
+		virtual bool OnXRHMDTransform(FTransform HMDTransform) override;
+		virtual bool OnXRControllerTransform(FTransform ControllerTransform, EControllerHand Handedness) override;
+		virtual bool OnXRButtonTouched(EControllerHand Handedness, FKey Key, bool bIsRepeat) override;
+		virtual bool OnXRButtonTouchReleased(EControllerHand Handedness, FKey Key) override;
+		virtual bool OnXRButtonPressed(EControllerHand Handedness, FKey Key, bool bIsRepeat) override;
+		virtual bool OnXRButtonReleased(EControllerHand Handedness, FKey Key) override;
+		virtual bool OnXRAnalog(EControllerHand Handedness, FKey Key, double AnalogValue) override;
+		virtual bool OnXRSystem(EPixelStreaming2XRSystem System) override;
+
+	protected:
+		FIntPoint	ConvertFromNormalizedScreenLocation(const FVector2D& ScreenLocation, bool bIncludeOffset = true);
+		FWidgetPath FindRoutingMessageWidget(const FVector2D& Location) const;
+		FKey		TranslateMouseButtonToKey(const EMouseButtons::Type Button);
+
+		struct FCachedTouchEvent
+		{
+			FVector2D Location;
+			float	  Force;
+			int32	  ControllerIndex;
+		};
+
+		struct FAnalogValue
+		{
+			/** The actual analog value from the controller axis, typical 0.0..1.0 */
+			double Value;
+			/** If value is non-zero then keep applying this analog values across frames.
+			 * This is useful for trigger axis inputs where if a value is not transmitted
+			 * UE will assume a gap in input means a full trigger press (which is not accurate if we were still pressing).
+			 */
+			bool bKeepUnlessZero = false;
+			/** Has this key event already been fired once? */
+			bool bIsRepeat = false;
+		};
+
+		// Keep a cache of the last touch events as we need to fire Touch Moved every frame while touch is down
+		TMap<int32, FCachedTouchEvent> CachedTouchEvents;
+
+		using FKeyId = uint8;
+		/**
+		 * If more values are received in a single tick (e.g. could be temp network issue),
+		 * then we only forward the latest value.
+		 *
+		 * Reason: The input system seems to expect at most one raw analog value per FKey per Tick.
+		 * If this is not done, the input system can get stuck on non-zero input value even if the user has
+		 * already stopped moving the analog stick. It would stay stuck until the next time the user moves the stick.
+		 *
+		 * The values arrive in the order of recording: that means once the player releases the analog,
+		 * the last analog value would be 0.
+		 */
+		TMap<FInputDeviceId, TMap<FKey, FAnalogValue>> AnalogEventsReceivedThisTick;
+
+		/** Forwards the latest analog input received for each key this tick. */
+		void ProcessLatestAnalogInputFromThisTick();
+
+		// Track which touch events we processed this frame so we can avoid re-processing them
+		TSet<int32> TouchIndicesProcessedThisFrame;
+
+		// Sends Touch Moved events for any touch index which is currently down but wasn't already updated this frame
+		void BroadcastActiveTouchMoveEvents();
+
+		void FindFocusedWidget();
+
+		TArray<FKey> FilteredKeys;
+		void		 OnFilteredKeysChanged(IConsoleVariable* Var);
+		bool		 FilterKey(const FKey& Key);
+
+		struct FMessage
+		{
+			FString									 SourceId;
+			TFunction<void(FString, FMemoryReader)>* Handler;
+			TArray<uint8>							 Data;
+		};
+
+		TWeakPtr<SWindow>			  TargetWindow;
+		TWeakPtr<SViewport>			  TargetViewport;
+		TWeakPtr<FIntPoint>			  TargetScreenSize; // Deprecated functionality but remaining until it can be removed
+		TWeakPtr<FIntRect>			  TargetScreenRect; // Manual size override used when we don't have a single window/viewport target
+		uint8						  NumActiveTouches;
+		bool						  bIsMouseActive;
+		TQueue<FMessage>			  Messages;
+		EPixelStreaming2InputType	  InputType = EPixelStreaming2InputType::RouteToWindow;
+		FVector2D					  LastTouchLocation = FVector2D(EForceInit::ForceInitToZero);
+		TMap<uint8, MessageHandlerFn> DispatchTable;
+
+		/** Reference to the message handler which events should be passed to. */
+		TSharedPtr<FGenericApplicationMessageHandler> MessageHandler;
+
+		/** For convenience, we keep a reference to the application wrapper owned by the input channel */
+		TSharedPtr<class FPixelStreaming2ApplicationWrapper> PixelStreamerApplicationWrapper;
+
+		/**
+		 * Is the application faking touch events by dragging the mouse along
+		 * the canvas? If so then we must put the browser canvas in a special
+		 * state to replicate the behavior of the application.
+		 */
+		bool bFakingTouchEvents;
+
+		/**
+		 * Padding for string parsing when handling messages.
+		 * 1 character for the actual message and then
+		 * 2 characters for the length which are skipped
+		 */
+		const size_t MessageHeaderOffset = 1;
+
+		struct FPixelStreaming2XRController
+		{
+		public:
+			FTransform		Transform;
+			EControllerHand Handedness;
+		};
+
+		TMap<EControllerHand, FPixelStreaming2XRController> XRControllers;
+
+		/**
+		 * A map of named commands we respond to when we receive a datachannel message of type "command".
+		 * Key = command name (e.g "Encoder.MaxQP")
+		 * Value = The command handler lambda function whose parameters are as follows:
+		 *  FString - the source id of the user who sent the message
+		 *  FString - the descriptor (e.g. the full json payload of the command message)
+		 *  FString - the parsed value of the command, e.g. if key was "Encoder.MaxQP" and descriptor was { type: "Command", "Encoder.MaxQP": 51 }, then parsed value is "51".
+		 */
+		TMap<FString, CommandHandlerFn> CommandHandlers;
+
+		TFunction<bool(FString)> ElevatedCheck;
+
+		TSharedPtr<IPixelStreaming2DataProtocol> ToStreamerProtocol;
+		TSharedPtr<IPixelStreaming2DataProtocol> FromStreamerProtocol;
+
+		float uint16_MAX = (float)UINT16_MAX;
+		float int16_MAX = (float)SHRT_MAX;
+
+	private:
+		FDelegateHandle OnInputKeyFilterChangedHandle;
+	};
+} // namespace UE::PixelStreaming2Input
\ No newline at end of file
diff --git a/Engine/Plugins/Media/PixelStreaming2/Source/PixelStreaming2Input/Internal/InputStructures.h b/Engine/Plugins/Media/PixelStreaming2/Source/PixelStreaming2Input/Internal/InputStructures.h
new file mode 100644
--- /dev/null
+++ b/Engine/Plugins/Media/PixelStreaming2/Source/PixelStreaming2Input/Internal/InputStructures.h
@@ -0,0 +1,34 @@
+// Copyright Epic Games, Inc. All Rights Reserved.
+
+#pragma once
+
+#include "CoreMinimal.h"
+#include "Serialization/BufferArchive.h"
+#include "Serialization/MemoryReader.h"
+
+namespace UE::PixelStreaming2Input
+{
+	template <typename IntegerSequence>
+	struct TPayloadParser;
+
+	template <uint32... Indices>
+	struct TPayloadParser<TIntegerSequence<uint32, Indices...>>
+	{
+		template <typename TupleType>
+		static void Parse(TupleType&& Tuple, FArchive& Ar)
+		{
+			([](auto& Arg, FArchive& Ar) { Ar << Arg; }(Forward<TupleType>(Tuple).template Get<Indices>(), Ar), ...);
+		}
+	};
+
+	template <typename... Types>
+	struct TPayload : public TTuple<std::decay_t<Types>...>
+	{
+		TPayload(FArchive& Ar)
+		{
+			TPayloadParser<TMakeIntegerSequence<uint32, sizeof...(Types)>>::Parse(*this, Ar);
+
+			checkf(Ar.GetArchiveState().AtEnd(), TEXT("Archive isn't empty after parsing into payload. You may be losing information!"));
+		}
+	};
+} // namespace UE::PixelStreaming2Input
diff --git a/Engine/Plugins/Media/PixelStreaming2/Source/PixelStreaming2Input/Internal/JavaScriptKeyCodes.inl b/Engine/Plugins/Media/PixelStreaming2/Source/PixelStreaming2Input/Internal/JavaScriptKeyCodes.inl
new file mode 100644
--- /dev/null
+++ b/Engine/Plugins/Media/PixelStreaming2/Source/PixelStreaming2Input/Internal/JavaScriptKeyCodes.inl
@@ -0,0 +1,280 @@
+// Copyright Epic Games, Inc. All Rights Reserved.
+
+#pragma once
+
+#include "InputCoreTypes.h"
+
+namespace UE::PixelStreaming2Input 
+{
+
+/**
+ * An array which will convert JavaScript key codes into FKeys. The JavaScript
+ * key code should be used as an index into the array to perform the lookup.
+ * https://www.cambiaresearch.com/articles/15/javascript-char-codes-key-codes
+ * http://keycode.info/
+ * Most key codes use the standard values given in the links above, but some
+ * have been choosen specially for 'right' versions of the same key.
+ */
+
+	static const FKey* JavaScriptKeyCodeToFKey[256] =
+	{
+		/*   0 */ &EKeys::Invalid,
+		/*   1 */ &EKeys::Invalid,
+		/*   2 */ &EKeys::Invalid,
+		/*   3 */ &EKeys::Invalid,
+		/*   4 */ &EKeys::Invalid,
+		/*   5 */ &EKeys::Invalid,
+		/*   6 */ &EKeys::Invalid,
+		/*   7 */ &EKeys::Invalid,
+		/*   8 */ &EKeys::BackSpace,
+		/*   9 */ &EKeys::Tab,
+		/*  10 */ &EKeys::Invalid,
+		/*  11 */ &EKeys::Invalid,
+		/*  12 */ /*&EKeys::Clear*/ &EKeys::Invalid,
+		/*  13 */ &EKeys::Enter,
+		/*  14 */ &EKeys::Invalid,
+		/*  15 */ &EKeys::Invalid,
+		/*  16 */ &EKeys::LeftShift,
+		/*  17 */ &EKeys::LeftControl,
+		/*  18 */ &EKeys::LeftAlt,
+		/*  19 */ &EKeys::Pause,
+		/*  20 */ &EKeys::CapsLock,
+		/*  21 */ &EKeys::Invalid,
+		/*  22 */ &EKeys::Invalid,
+		/*  23 */ &EKeys::Invalid,
+		/*  24 */ &EKeys::Invalid,
+		/*  25 */ &EKeys::Invalid,
+		/*  26 */ &EKeys::Invalid,
+		/*  27 */ &EKeys::Escape,
+		/*  28 */ &EKeys::Invalid,
+		/*  29 */ &EKeys::Invalid,
+		/*  30 */ &EKeys::Invalid,
+		/*  31 */ &EKeys::Invalid,
+		/*  32 */ &EKeys::SpaceBar,
+		/*  33 */ &EKeys::PageUp,
+		/*  34 */ &EKeys::PageDown,
+		/*  35 */ &EKeys::End,
+		/*  36 */ &EKeys::Home,
+		/*  37 */ &EKeys::Left,
+		/*  38 */ &EKeys::Up,
+		/*  39 */ &EKeys::Right,
+		/*  40 */ &EKeys::Down,
+		/*  41 */ &EKeys::Invalid,
+		/*  42 */ &EKeys::Invalid,
+		/*  43 */ &EKeys::Invalid,
+		/*  44 */ /*&EKeys::PrintScreen*/ &EKeys::Invalid,
+		/*  45 */ &EKeys::Insert,
+		/*  46 */ &EKeys::Delete,
+		/*  47 */ &EKeys::Invalid,
+		/*  48 */ &EKeys::Zero,
+		/*  49 */ &EKeys::One,
+		/*  50 */ &EKeys::Two,
+		/*  51 */ &EKeys::Three,
+		/*  52 */ &EKeys::Four,
+		/*  53 */ &EKeys::Five,
+		/*  54 */ &EKeys::Six,
+		/*  55 */ &EKeys::Seven,
+		/*  56 */ &EKeys::Eight,
+		/*  57 */ &EKeys::Nine,
+		/*  58 */ &EKeys::Invalid,
+		/*  59 */ &EKeys::Invalid,
+		/*  60 */ &EKeys::Invalid,
+		/*  61 */ &EKeys::Invalid,
+		/*  62 */ &EKeys::Invalid,
+		/*  63 */ &EKeys::Invalid,
+		/*  64 */ &EKeys::Invalid,
+		/*  65 */ &EKeys::A,
+		/*  66 */ &EKeys::B,
+		/*  67 */ &EKeys::C,
+		/*  68 */ &EKeys::D,
+		/*  69 */ &EKeys::E,
+		/*  70 */ &EKeys::F,
+		/*  71 */ &EKeys::G,
+		/*  72 */ &EKeys::H,
+		/*  73 */ &EKeys::I,
+		/*  74 */ &EKeys::J,
+		/*  75 */ &EKeys::K,
+		/*  76 */ &EKeys::L,
+		/*  77 */ &EKeys::M,
+		/*  78 */ &EKeys::N,
+		/*  79 */ &EKeys::O,
+		/*  80 */ &EKeys::P,
+		/*  81 */ &EKeys::Q,
+		/*  82 */ &EKeys::R,
+		/*  83 */ &EKeys::S,
+		/*  84 */ &EKeys::T,
+		/*  85 */ &EKeys::U,
+		/*  86 */ &EKeys::V,
+		/*  87 */ &EKeys::W,
+		/*  88 */ &EKeys::X,
+		/*  89 */ &EKeys::Y,
+		/*  90 */ &EKeys::Z,
+		/*  91 */ /*&EKeys::LeftWindowKey*/ &EKeys::Invalid,
+		/*  92 */ /*&EKeys::RightWindowKey*/ &EKeys::Invalid,
+		/*  93 */ /*&EKeys::SelectKey*/ &EKeys::Invalid,
+		/*  94 */ &EKeys::Invalid,
+		/*  95 */ &EKeys::Invalid,
+		/*  96 */ &EKeys::NumPadZero,
+		/*  97 */ &EKeys::NumPadOne,
+		/*  98 */ &EKeys::NumPadTwo,
+		/*  99 */ &EKeys::NumPadThree,
+		/* 100 */ &EKeys::NumPadFour,
+		/* 101 */ &EKeys::NumPadFive,
+		/* 102 */ &EKeys::NumPadSix,
+		/* 103 */ &EKeys::NumPadSeven,
+		/* 104 */ &EKeys::NumPadEight,
+		/* 105 */ &EKeys::NumPadNine,
+		/* 106 */ &EKeys::Multiply,
+		/* 107 */ &EKeys::Add,
+		/* 108 */ &EKeys::Invalid,
+		/* 109 */ &EKeys::Subtract,
+		/* 110 */ &EKeys::Decimal,
+		/* 111 */ &EKeys::Divide,
+		/* 112 */ &EKeys::F1,
+		/* 113 */ &EKeys::F2,
+		/* 114 */ &EKeys::F3,
+		/* 115 */ &EKeys::F4,
+		/* 116 */ &EKeys::F5,
+		/* 117 */ &EKeys::F6,
+		/* 118 */ &EKeys::F7,
+		/* 119 */ &EKeys::F8,
+		/* 120 */ &EKeys::F9,
+		/* 121 */ &EKeys::F10,
+		/* 122 */ &EKeys::F11,
+		/* 123 */ &EKeys::F12,
+		/* 124 */ &EKeys::Invalid,
+		/* 125 */ &EKeys::Invalid,
+		/* 126 */ &EKeys::Invalid,
+		/* 127 */ &EKeys::Invalid,
+		/* 128 */ &EKeys::Invalid,
+		/* 129 */ &EKeys::Invalid,
+		/* 130 */ &EKeys::Invalid,
+		/* 131 */ &EKeys::Invalid,
+		/* 132 */ &EKeys::Invalid,
+		/* 133 */ &EKeys::Invalid,
+		/* 134 */ &EKeys::Invalid,
+		/* 135 */ &EKeys::Invalid,
+		/* 136 */ &EKeys::Invalid,
+		/* 137 */ &EKeys::Invalid,
+		/* 138 */ &EKeys::Invalid,
+		/* 139 */ &EKeys::Invalid,
+		/* 140 */ &EKeys::Invalid,
+		/* 141 */ &EKeys::Invalid,
+		/* 142 */ &EKeys::Invalid,
+		/* 143 */ &EKeys::Invalid,
+		/* 144 */ &EKeys::NumLock,
+		/* 145 */ &EKeys::ScrollLock,
+		/* 146 */ &EKeys::Invalid,
+		/* 147 */ &EKeys::Invalid,
+		/* 148 */ &EKeys::Invalid,
+		/* 149 */ &EKeys::Invalid,
+		/* 150 */ &EKeys::Invalid,
+		/* 151 */ &EKeys::Invalid,
+		/* 152 */ &EKeys::Invalid,
+		/* 153 */ &EKeys::Invalid,
+		/* 154 */ &EKeys::Invalid,
+		/* 155 */ &EKeys::Invalid,
+		/* 156 */ &EKeys::Invalid,
+		/* 157 */ &EKeys::Invalid,
+		/* 158 */ &EKeys::Invalid,
+		/* 159 */ &EKeys::Invalid,
+		/* 160 */ &EKeys::Invalid,
+		/* 161 */ &EKeys::Invalid,
+		/* 162 */ &EKeys::Invalid,
+		/* 163 */ &EKeys::Invalid,
+		/* 164 */ &EKeys::Invalid,
+		/* 165 */ &EKeys::Invalid,
+		/* 166 */ &EKeys::Invalid,
+		/* 167 */ &EKeys::Invalid,
+		/* 168 */ &EKeys::Invalid,
+		/* 169 */ &EKeys::Invalid,
+		/* 170 */ &EKeys::Invalid,
+		/* 171 */ &EKeys::Invalid,
+		/* 172 */ &EKeys::Invalid,
+		/* 173 */ &EKeys::Invalid,
+		/* 174 */ &EKeys::Invalid,
+		/* 175 */ &EKeys::Invalid,
+		/* 176 */ &EKeys::Invalid,
+		/* 177 */ &EKeys::Invalid,
+		/* 178 */ &EKeys::Invalid,
+		/* 179 */ &EKeys::Invalid,
+		/* 180 */ &EKeys::Invalid,
+		/* 181 */ &EKeys::Invalid,
+		/* 182 */ &EKeys::Invalid,
+		/* 183 */ &EKeys::Invalid,
+		/* 184 */ &EKeys::Invalid,
+		/* 185 */ &EKeys::Invalid,
+		/* 186 */ &EKeys::Semicolon,
+		/* 187 */ &EKeys::Equals,
+		/* 188 */ &EKeys::Comma,
+		/* 189 */ &EKeys::Hyphen,
+		/* 190 */ &EKeys::Period,
+		/* 191 */ &EKeys::Slash,
+		/* 192 */ &EKeys::Tilde,
+		/* 193 */ &EKeys::Invalid,
+		/* 194 */ &EKeys::Invalid,
+		/* 195 */ &EKeys::Invalid,
+		/* 196 */ &EKeys::Invalid,
+		/* 197 */ &EKeys::Invalid,
+		/* 198 */ &EKeys::Invalid,
+		/* 199 */ &EKeys::Invalid,
+		/* 200 */ &EKeys::Invalid,
+		/* 201 */ &EKeys::Invalid,
+		/* 202 */ &EKeys::Invalid,
+		/* 203 */ &EKeys::Invalid,
+		/* 204 */ &EKeys::Invalid,
+		/* 205 */ &EKeys::Invalid,
+		/* 206 */ &EKeys::Invalid,
+		/* 207 */ &EKeys::Invalid,
+		/* 208 */ &EKeys::Invalid,
+		/* 209 */ &EKeys::Invalid,
+		/* 210 */ &EKeys::Invalid,
+		/* 211 */ &EKeys::Invalid,
+		/* 212 */ &EKeys::Invalid,
+		/* 213 */ &EKeys::Invalid,
+		/* 214 */ &EKeys::Invalid,
+		/* 215 */ &EKeys::Invalid,
+		/* 216 */ &EKeys::Invalid,
+		/* 217 */ &EKeys::Invalid,
+		/* 218 */ &EKeys::Invalid,
+		/* 219 */ &EKeys::LeftBracket,
+		/* 220 */ &EKeys::Backslash,
+		/* 221 */ &EKeys::RightBracket,
+		/* 222 */ &EKeys::Apostrophe,
+		/* 223 */ &EKeys::Tilde,
+		/* 224 */ &EKeys::Invalid,
+		/* 225 */ &EKeys::Invalid,
+		/* 226 */ &EKeys::Invalid,
+		/* 227 */ &EKeys::Invalid,
+		/* 228 */ &EKeys::Invalid,
+		/* 229 */ &EKeys::Invalid,
+		/* 230 */ &EKeys::Invalid,
+		/* 231 */ &EKeys::Invalid,
+		/* 232 */ &EKeys::Invalid,
+		/* 233 */ &EKeys::Invalid,
+		/* 234 */ &EKeys::Invalid,
+		/* 235 */ &EKeys::Invalid,
+		/* 236 */ &EKeys::Invalid,
+		/* 237 */ &EKeys::Invalid,
+		/* 238 */ &EKeys::Invalid,
+		/* 239 */ &EKeys::Invalid,
+		/* 240 */ &EKeys::Invalid,
+		/* 241 */ &EKeys::Invalid,
+		/* 242 */ &EKeys::Invalid,
+		/* 243 */ &EKeys::Invalid,
+		/* 244 */ &EKeys::Invalid,
+		/* 245 */ &EKeys::Invalid,
+		/* 246 */ &EKeys::Invalid,
+		/* 247 */ &EKeys::Invalid,
+		/* 248 */ &EKeys::Invalid,
+		/* 249 */ &EKeys::Invalid,
+		/* 250 */ &EKeys::Invalid,
+		/* 251 */ &EKeys::Invalid,
+		/* 252 */ &EKeys::Invalid,
+
+		// Special key codes.
+		/* 253 */ &EKeys::RightShift,
+		/* 254 */ &EKeys::RightControl,
+		/* 255 */ &EKeys::RightAlt
+	};
+} // namespace UE::PixelStreaming2Input
\ No newline at end of file
diff --git a/Engine/Plugins/Media/PixelStreaming2/Source/PixelStreaming2Input/Private/InputConversion.cpp b/Engine/Plugins/Media/PixelStreaming2/Source/PixelStreaming2Input/Private/InputConversion.cpp
deleted file mode 100644
diff --git a/Engine/Plugins/Media/PixelStreaming2/Source/PixelStreaming2Input/Private/InputConversion.h b/Engine/Plugins/Media/PixelStreaming2/Source/PixelStreaming2Input/Private/InputConversion.h
deleted file mode 100644
diff --git a/Engine/Plugins/Media/PixelStreaming2/Source/PixelStreaming2Input/Private/InputDevice.cpp b/Engine/Plugins/Media/PixelStreaming2/Source/PixelStreaming2Input/Private/InputDevice.cpp
--- a/Engine/Plugins/Media/PixelStreaming2/Source/PixelStreaming2Input/Private/InputDevice.cpp
+++ b/Engine/Plugins/Media/PixelStreaming2/Source/PixelStreaming2Input/Private/InputDevice.cpp
@@ -25,65 +25,55 @@ namespace UE::PixelStreaming2Input
 		FSlateApplication::Get().OnPreTick().AddRaw(this, &FInputDevice::Tick);
 	}
 
-	void FInputDevice::AddInputHandler(TSharedPtr<IPixelStreaming2InputHandler> InputHandler)
+	void FInputDevice::AddInputHandler(IPixelStreaming2InputHandler* InputHandler)
 	{
 		InputHandlers.Add(InputHandler);
 	}
 
+	void FInputDevice::RemoveInputHandler(IPixelStreaming2InputHandler* InputHandler)
+	{
+		InputHandlers.Remove(InputHandler);
+	}
+
 	void FInputDevice::Tick(float DeltaTime)
 	{
-		for (TWeakPtr<IPixelStreaming2InputHandler> WeakHandler : InputHandlers)
+		for (IPixelStreaming2InputHandler* Handler : InputHandlers)
 		{
-			if (TSharedPtr<IPixelStreaming2InputHandler> Handler = WeakHandler.Pin())
-			{
-				Handler->Tick(DeltaTime);
-			}
+			Handler->Tick(DeltaTime);
 		}
 	}
 
 	void FInputDevice::SetMessageHandler(const TSharedRef<FGenericApplicationMessageHandler>& InTargetHandler)
 	{
-		for (TWeakPtr<IPixelStreaming2InputHandler> WeakHandler : InputHandlers)
+		for (IPixelStreaming2InputHandler* Handler : InputHandlers)
 		{
-			if (TSharedPtr<IPixelStreaming2InputHandler> Handler = WeakHandler.Pin())
-			{
-				Handler->SetMessageHandler(InTargetHandler);
-			}
+			Handler->SetMessageHandler(InTargetHandler);
 		}
 	}
 
 	bool FInputDevice::Exec(UWorld* InWorld, const TCHAR* Cmd, FOutputDevice& Ar)
 	{
 		bool bRetVal = true;
-		for (TWeakPtr<IPixelStreaming2InputHandler> WeakHandler : InputHandlers)
+		for (IPixelStreaming2InputHandler* Handler : InputHandlers)
 		{
-			if (TSharedPtr<IPixelStreaming2InputHandler> Handler = WeakHandler.Pin())
-			{
-				bRetVal &= Handler->Exec(InWorld, Cmd, Ar);
-			}
+			bRetVal &= Handler->Exec(InWorld, Cmd, Ar);
 		}
 		return bRetVal;
 	}
 
 	void FInputDevice::SetChannelValue(int32 ControllerId, FForceFeedbackChannelType ChannelType, float Value)
 	{
-		for (TWeakPtr<IPixelStreaming2InputHandler> WeakHandler : InputHandlers)
+		for (IPixelStreaming2InputHandler* Handler : InputHandlers)
 		{
-			if (TSharedPtr<IPixelStreaming2InputHandler> Handler = WeakHandler.Pin())
-			{
-				Handler->SetChannelValue(ControllerId, ChannelType, Value);
-			}
+			Handler->SetChannelValue(ControllerId, ChannelType, Value);
 		}
 	}
 
 	void FInputDevice::SetChannelValues(int32 ControllerId, const FForceFeedbackValues& Values)
 	{
-		for (TWeakPtr<IPixelStreaming2InputHandler> WeakHandler : InputHandlers)
+		for (IPixelStreaming2InputHandler* Handler : InputHandlers)
 		{
-			if (TSharedPtr<IPixelStreaming2InputHandler> Handler = WeakHandler.Pin())
-			{
-				Handler->SetChannelValues(ControllerId, Values);
-			}
+			Handler->SetChannelValues(ControllerId, Values);
 		}
 	}
 
diff --git a/Engine/Plugins/Media/PixelStreaming2/Source/PixelStreaming2Input/Private/InputDevice.h b/Engine/Plugins/Media/PixelStreaming2/Source/PixelStreaming2Input/Private/InputDevice.h
--- a/Engine/Plugins/Media/PixelStreaming2/Source/PixelStreaming2Input/Private/InputDevice.h
+++ b/Engine/Plugins/Media/PixelStreaming2/Source/PixelStreaming2Input/Private/InputDevice.h
@@ -27,7 +27,8 @@ namespace UE::PixelStreaming2Input
 		virtual void SetChannelValue(int32 ControllerId, FForceFeedbackChannelType ChannelType, float Value) override;
 		virtual void SetChannelValues(int32 ControllerId, const FForceFeedbackValues& Values) override;
 
-		void  AddInputHandler(TSharedPtr<IPixelStreaming2InputHandler> InputHandler);
+		void  AddInputHandler(IPixelStreaming2InputHandler* InputHandler);
+		void  RemoveInputHandler(IPixelStreaming2InputHandler* InputHandler);
 		uint8 OnControllerConnected();
 		void  OnControllerDisconnected(uint8 DeleteControllerId);
 		bool  GetPlatformUserAndDevice(uint8 ControllerId, FInputDeviceId& OutDeviceId, FPlatformUserId& OutPlatformUser);
@@ -48,10 +49,10 @@ namespace UE::PixelStreaming2Input
 		/**
 		 * The array of input handlers. Each input handler belongs to a single streamer
 		 *
-		 * NOTE: We store each input handler as a weakptr as we don't want to be the reason
-		 * the handler doesn't get deleted. Each handler should be tied to a Streamer's lifecycle
+		 * NOTE: We store each input handler as a raw pointer as registering and deregistering is handled in the FPixelStreaming2InputHandler ctor and dtor
+		 * NOTE: This pointer is non-owning. It is up to user code to keep the original IPixelStreaming2InputHandler alive to avoid deregistration
 		 */
-		TArray<TWeakPtr<IPixelStreaming2InputHandler>> InputHandlers;
+		TSet<IPixelStreaming2InputHandler*> InputHandlers;
 
 		/**
 		 * The map of connected controllers. As each handler can have multiple input devices, we want to make sure that each controller of each device is unique
diff --git a/Engine/Plugins/Media/PixelStreaming2/Source/PixelStreaming2Input/Private/InputHandler.cpp b/Engine/Plugins/Media/PixelStreaming2/Source/PixelStreaming2Input/Private/InputHandler.cpp
--- a/Engine/Plugins/Media/PixelStreaming2/Source/PixelStreaming2Input/Private/InputHandler.cpp
+++ b/Engine/Plugins/Media/PixelStreaming2/Source/PixelStreaming2Input/Private/InputHandler.cpp
@@ -2,6 +2,7 @@
 
 #include "InputHandler.h"
 
+#include "ApplicationWrapper.h"
 #include "InputStructures.h"
 #include "JavaScriptKeyCodes.inl"
 #include "Layout/ArrangedChildren.h"
@@ -19,13 +20,13 @@
 #include "IPixelStreaming2HMD.h"
 #include "IPixelStreaming2HMDModule.h"
 #include "PixelStreaming2InputEnums.h"
-#include "InputConversion.h"
 #include "Widgets/Input/SEditableText.h"
 #include "Widgets/Input/SMultiLineEditableTextBox.h"
 #include "InputDevice.h"
 #include "DefaultDataProtocol.h"
 #include "Logging.h"
 #include "PixelStreaming2PluginSettings.h"
+#include "PixelStreaming2InputModule.h"
 
 #if WITH_EDITOR
 	#include "UObject/UObjectGlobals.h"
@@ -36,162 +37,34 @@
 
 namespace UE::PixelStreaming2Input
 {
-	typedef EPixelStreaming2InputAction Action;
-	typedef EPixelStreaming2XRSystem	XRSystem;
-
-	FPixelStreaming2InputHandler::FPixelStreaming2InputHandler(TSharedPtr<FPixelStreaming2ApplicationWrapper> InApplicationWrapper, const TSharedPtr<FGenericApplicationMessageHandler>& InTargetHandler)
+	FPixelStreaming2InputHandler::FPixelStreaming2InputHandler()
 		: TargetViewport(nullptr)
 		, NumActiveTouches(0)
 		, bIsMouseActive(false)
-		, MessageHandler(InTargetHandler)
-		, PixelStreamerApplicationWrapper(InApplicationWrapper)
-		, ToStreamerProtocol(GetDefaultToStreamerProtocol())
-		, FromStreamerProtocol(GetDefaultFromStreamerProtocol())
+		, MessageHandler(FSlateApplication::Get().GetPlatformApplication()->GetMessageHandler())
+		, PixelStreamerApplicationWrapper(MakeShared<FPixelStreaming2ApplicationWrapper>(FSlateApplication::Get().GetPlatformApplication()))
+
 	{
 		// Register this input handler as an IMotionController. The module handles the registering as an IInputDevice
 		IModularFeatures::Get().RegisterModularFeature(GetModularFeatureName(), this);
 
-		RegisterMessageHandler("KeyPress", [this](FString SourceId, FMemoryReader Ar) { HandleOnKeyChar(Ar); });
-		RegisterMessageHandler("KeyUp", [this](FString SourceId, FMemoryReader Ar) { HandleOnKeyUp(Ar); });
-		RegisterMessageHandler("KeyDown", [this](FString SourceId, FMemoryReader Ar) { HandleOnKeyDown(Ar); });
-
-		RegisterMessageHandler("TouchStart", [this](FString SourceId, FMemoryReader Ar) { HandleOnTouchStarted(Ar); });
-		RegisterMessageHandler("TouchMove", [this](FString SourceId, FMemoryReader Ar) { HandleOnTouchMoved(Ar); });
-		RegisterMessageHandler("TouchEnd", [this](FString SourceId, FMemoryReader Ar) { HandleOnTouchEnded(Ar); });
-
-		RegisterMessageHandler("GamepadConnected", [this](FString SourceId, FMemoryReader Ar) { HandleOnControllerConnected(Ar); });
-		RegisterMessageHandler("GamepadAnalog", [this](FString SourceId, FMemoryReader Ar) { HandleOnControllerAnalog(Ar); });
-		RegisterMessageHandler("GamepadButtonPressed", [this](FString SourceId, FMemoryReader Ar) { HandleOnControllerButtonPressed(Ar); });
-		RegisterMessageHandler("GamepadButtonReleased", [this](FString SourceId, FMemoryReader Ar) { HandleOnControllerButtonReleased(Ar); });
-		RegisterMessageHandler("GamepadDisconnected", [this](FString SourceId, FMemoryReader Ar) { HandleOnControllerDisconnected(Ar); });
-
-		RegisterMessageHandler("MouseEnter", [this](FString SourceId, FMemoryReader Ar) { HandleOnMouseEnter(Ar); });
-		RegisterMessageHandler("MouseLeave", [this](FString SourceId, FMemoryReader Ar) { HandleOnMouseLeave(Ar); });
-		RegisterMessageHandler("MouseUp", [this](FString SourceId, FMemoryReader Ar) { HandleOnMouseUp(Ar); });
-		RegisterMessageHandler("MouseDown", [this](FString SourceId, FMemoryReader Ar) { HandleOnMouseDown(Ar); });
-		RegisterMessageHandler("MouseMove", [this](FString SourceId, FMemoryReader Ar) { HandleOnMouseMove(Ar); });
-		RegisterMessageHandler("MouseWheel", [this](FString SourceId, FMemoryReader Ar) { HandleOnMouseWheel(Ar); });
-		RegisterMessageHandler("MouseDouble", [this](FString SourceId, FMemoryReader Ar) { HandleOnMouseDoubleClick(Ar); });
-
-		RegisterMessageHandler("XREyeViews", [this](FString SourceId, FMemoryReader Ar) { HandleOnXREyeViews(Ar); });
-		RegisterMessageHandler("XRHMDTransform", [this](FString SourceId, FMemoryReader Ar) { HandleOnXRHMDTransform(Ar); });
-		RegisterMessageHandler("XRControllerTransform", [this](FString SourceId, FMemoryReader Ar) { HandleOnXRControllerTransform(Ar); });
-		RegisterMessageHandler("XRButtonPressed", [this](FString SourceId, FMemoryReader Ar) { HandleOnXRButtonPressed(Ar); });
-		RegisterMessageHandler("XRButtonTouched", [this](FString SourceId, FMemoryReader Ar) { HandleOnXRButtonTouched(Ar); });
-		RegisterMessageHandler("XRButtonTouchReleased", [this](FString SourceId, FMemoryReader Ar) { HandleOnXRButtonTouchReleased(Ar); });
-		RegisterMessageHandler("XRButtonReleased", [this](FString SourceId, FMemoryReader Ar) { HandleOnXRButtonReleased(Ar); });
-		RegisterMessageHandler("XRAnalog", [this](FString SourceId, FMemoryReader Ar) { HandleOnXRAnalog(Ar); });
-		RegisterMessageHandler("XRSystem", [this](FString SourceId, FMemoryReader Ar) { HandleOnXRSystem(Ar); });
-
-		RegisterMessageHandler("Command", [this](FString SourceId, FMemoryReader Ar) { HandleOnCommand(SourceId, Ar); });
-		RegisterMessageHandler("UIInteraction", [this](FString SourceId, FMemoryReader Ar) { HandleUIInteraction(Ar); });
-		RegisterMessageHandler("TextboxEntry", [this](FString SourceId, FMemoryReader Ar) { HandleOnTextboxEntry(Ar); });
-
-		// RequestQualityControl has been removed. We no-op this handler to prevent warnings about unregistered message types from older front-end versions
-		RegisterMessageHandler("RequestQualityControl", [this](FString SourceId, FMemoryReader Ar) { /* Do nothing */ });
-
-		// Populate map
-		// Button indices found in: [REDACTED_PROFILE_URL]/webxr-input-profiles/tree/master/packages/registry/profiles
-		// HTC Vive - Left Hand
-		// Buttons
-		FInputConverter::XRInputToFKey.Add(MakeTuple(XRSystem::HTCVive, EControllerHand::Left, 0, Action::Click), EKeys::Vive_Left_Trigger_Click);
-		FInputConverter::XRInputToFKey.Add(MakeTuple(XRSystem::HTCVive, EControllerHand::Left, 0, Action::Axis), EKeys::Vive_Left_Trigger_Axis);
-		FInputConverter::XRInputToFKey.Add(MakeTuple(XRSystem::HTCVive, EControllerHand::Left, 1, Action::Click), EKeys::Vive_Left_Grip_Click);
-		FInputConverter::XRInputToFKey.Add(MakeTuple(XRSystem::HTCVive, EControllerHand::Left, 2, Action::Click), EKeys::Vive_Left_Trackpad_Click);
-		FInputConverter::XRInputToFKey.Add(MakeTuple(XRSystem::HTCVive, EControllerHand::Left, 2, Action::Touch), EKeys::Vive_Left_Trackpad_Touch);
-		// Axes
-		FInputConverter::XRInputToFKey.Add(MakeTuple(XRSystem::HTCVive, EControllerHand::Left, 0, Action::X), EKeys::Vive_Left_Trackpad_X);
-		FInputConverter::XRInputToFKey.Add(MakeTuple(XRSystem::HTCVive, EControllerHand::Left, 1, Action::Y), EKeys::Vive_Left_Trackpad_Y);
-		// HTC Vive - Right Hand
-		// Buttons
-		FInputConverter::XRInputToFKey.Add(MakeTuple(XRSystem::HTCVive, EControllerHand::Right, 0, Action::Click), EKeys::Vive_Right_Trigger_Click);
-		FInputConverter::XRInputToFKey.Add(MakeTuple(XRSystem::HTCVive, EControllerHand::Right, 0, Action::Axis), EKeys::Vive_Right_Trigger_Axis);
-		FInputConverter::XRInputToFKey.Add(MakeTuple(XRSystem::HTCVive, EControllerHand::Right, 1, Action::Click), EKeys::Vive_Right_Grip_Click);
-		FInputConverter::XRInputToFKey.Add(MakeTuple(XRSystem::HTCVive, EControllerHand::Right, 2, Action::Click), EKeys::Vive_Right_Trackpad_Click);
-		FInputConverter::XRInputToFKey.Add(MakeTuple(XRSystem::HTCVive, EControllerHand::Right, 2, Action::Touch), EKeys::Vive_Right_Trackpad_Touch);
-		// Axes
-		FInputConverter::XRInputToFKey.Add(MakeTuple(XRSystem::HTCVive, EControllerHand::Right, 0, Action::X), EKeys::Vive_Right_Trackpad_X);
-		FInputConverter::XRInputToFKey.Add(MakeTuple(XRSystem::HTCVive, EControllerHand::Right, 1, Action::Y), EKeys::Vive_Right_Trackpad_Y);
-
-		// Quest - Left Hand
-		// Buttons
-		FInputConverter::XRInputToFKey.Add(MakeTuple(XRSystem::Quest, EControllerHand::Left, 0, Action::Click), EKeys::OculusTouch_Left_Trigger_Click);
-		FInputConverter::XRInputToFKey.Add(MakeTuple(XRSystem::Quest, EControllerHand::Left, 0, Action::Axis), EKeys::OculusTouch_Left_Trigger_Axis);
-		FInputConverter::XRInputToFKey.Add(MakeTuple(XRSystem::Quest, EControllerHand::Left, 0, Action::Touch), EKeys::OculusTouch_Left_Trigger_Touch);
-		FInputConverter::XRInputToFKey.Add(MakeTuple(XRSystem::Quest, EControllerHand::Left, 1, Action::Click), EKeys::OculusTouch_Left_Grip_Click);
-		FInputConverter::XRInputToFKey.Add(MakeTuple(XRSystem::Quest, EControllerHand::Left, 1, Action::Axis), EKeys::OculusTouch_Left_Grip_Axis);
-		// Index 1 (grip) touch not supported in UE
-		// Index 2 not supported by WebXR
-		FInputConverter::XRInputToFKey.Add(MakeTuple(XRSystem::Quest, EControllerHand::Left, 3, Action::Click), EKeys::OculusTouch_Left_Thumbstick_Click);
-		FInputConverter::XRInputToFKey.Add(MakeTuple(XRSystem::Quest, EControllerHand::Left, 3, Action::Touch), EKeys::OculusTouch_Left_Thumbstick_Touch);
-		FInputConverter::XRInputToFKey.Add(MakeTuple(XRSystem::Quest, EControllerHand::Left, 4, Action::Click), EKeys::OculusTouch_Left_X_Click);
-		FInputConverter::XRInputToFKey.Add(MakeTuple(XRSystem::Quest, EControllerHand::Left, 4, Action::Touch), EKeys::OculusTouch_Left_X_Touch);
-		FInputConverter::XRInputToFKey.Add(MakeTuple(XRSystem::Quest, EControllerHand::Left, 5, Action::Click), EKeys::OculusTouch_Left_Y_Click);
-		FInputConverter::XRInputToFKey.Add(MakeTuple(XRSystem::Quest, EControllerHand::Left, 5, Action::Touch), EKeys::OculusTouch_Left_Y_Touch);
-		// Index 6 (thumbrest) not supported in UE
-
-		// Axes
-		// Indices 0 and 1 not supported in WebXR
-		FInputConverter::XRInputToFKey.Add(MakeTuple(XRSystem::Quest, EControllerHand::Left, 2, Action::X), EKeys::OculusTouch_Left_Thumbstick_X);
-		FInputConverter::XRInputToFKey.Add(MakeTuple(XRSystem::Quest, EControllerHand::Left, 3, Action::Y), EKeys::OculusTouch_Left_Thumbstick_Y);
-
-		// Quest - Right Hand
-		// Buttons
-		FInputConverter::XRInputToFKey.Add(MakeTuple(XRSystem::Quest, EControllerHand::Right, 0, Action::Click), EKeys::OculusTouch_Right_Trigger_Click);
-		FInputConverter::XRInputToFKey.Add(MakeTuple(XRSystem::Quest, EControllerHand::Right, 0, Action::Axis), EKeys::OculusTouch_Right_Trigger_Axis);
-		FInputConverter::XRInputToFKey.Add(MakeTuple(XRSystem::Quest, EControllerHand::Right, 0, Action::Touch), EKeys::OculusTouch_Right_Trigger_Touch);
-		FInputConverter::XRInputToFKey.Add(MakeTuple(XRSystem::Quest, EControllerHand::Right, 1, Action::Click), EKeys::OculusTouch_Right_Grip_Click);
-		FInputConverter::XRInputToFKey.Add(MakeTuple(XRSystem::Quest, EControllerHand::Right, 1, Action::Axis), EKeys::OculusTouch_Right_Grip_Axis);
-		// Index 1 (grip) touch not supported in UE
-		// Index 2 not supported by WebXR
-		FInputConverter::XRInputToFKey.Add(MakeTuple(XRSystem::Quest, EControllerHand::Right, 3, Action::Click), EKeys::OculusTouch_Right_Thumbstick_Click);
-		FInputConverter::XRInputToFKey.Add(MakeTuple(XRSystem::Quest, EControllerHand::Right, 3, Action::Touch), EKeys::OculusTouch_Right_Thumbstick_Touch);
-		FInputConverter::XRInputToFKey.Add(MakeTuple(XRSystem::Quest, EControllerHand::Right, 4, Action::Click), EKeys::OculusTouch_Right_A_Click);
-		FInputConverter::XRInputToFKey.Add(MakeTuple(XRSystem::Quest, EControllerHand::Right, 4, Action::Touch), EKeys::OculusTouch_Right_A_Touch);
-		FInputConverter::XRInputToFKey.Add(MakeTuple(XRSystem::Quest, EControllerHand::Right, 5, Action::Click), EKeys::OculusTouch_Right_B_Click);
-		FInputConverter::XRInputToFKey.Add(MakeTuple(XRSystem::Quest, EControllerHand::Right, 5, Action::Touch), EKeys::OculusTouch_Right_B_Touch);
-		// Index 6 (thumbrest) not supported in UE
-
-		// Axes
-		// Indices 0 and 1 not supported in WebXR
-		FInputConverter::XRInputToFKey.Add(MakeTuple(XRSystem::Quest, EControllerHand::Right, 2, Action::X), EKeys::OculusTouch_Right_Thumbstick_X);
-		FInputConverter::XRInputToFKey.Add(MakeTuple(XRSystem::Quest, EControllerHand::Right, 3, Action::Y), EKeys::OculusTouch_Right_Thumbstick_Y);
-
-		// Gamepad Axes
-		FInputConverter::GamepadInputToFKey.Add(MakeTuple(1, Action::Axis), EKeys::Gamepad_LeftX);
-		FInputConverter::GamepadInputToFKey.Add(MakeTuple(2, Action::Axis), EKeys::Gamepad_LeftY);
-		FInputConverter::GamepadInputToFKey.Add(MakeTuple(3, Action::Axis), EKeys::Gamepad_RightX);
-		FInputConverter::GamepadInputToFKey.Add(MakeTuple(4, Action::Axis), EKeys::Gamepad_RightY);
-		FInputConverter::GamepadInputToFKey.Add(MakeTuple(5, Action::Axis), EKeys::Gamepad_LeftTriggerAxis);
-		FInputConverter::GamepadInputToFKey.Add(MakeTuple(6, Action::Axis), EKeys::Gamepad_RightTriggerAxis);
-		// Gamepad Buttons
-		FInputConverter::GamepadInputToFKey.Add(MakeTuple(0, Action::Click), EKeys::Gamepad_FaceButton_Bottom);
-		FInputConverter::GamepadInputToFKey.Add(MakeTuple(1, Action::Click), EKeys::Gamepad_FaceButton_Right);
-		FInputConverter::GamepadInputToFKey.Add(MakeTuple(2, Action::Click), EKeys::Gamepad_FaceButton_Left);
-		FInputConverter::GamepadInputToFKey.Add(MakeTuple(3, Action::Click), EKeys::Gamepad_FaceButton_Top);
-		FInputConverter::GamepadInputToFKey.Add(MakeTuple(4, Action::Click), EKeys::Gamepad_LeftShoulder);
-		FInputConverter::GamepadInputToFKey.Add(MakeTuple(5, Action::Click), EKeys::Gamepad_RightShoulder);
-		FInputConverter::GamepadInputToFKey.Add(MakeTuple(6, Action::Click), EKeys::Gamepad_LeftTrigger);
-		FInputConverter::GamepadInputToFKey.Add(MakeTuple(7, Action::Click), EKeys::Gamepad_RightTrigger);
-		FInputConverter::GamepadInputToFKey.Add(MakeTuple(8, Action::Click), EKeys::Gamepad_Special_Left);
-		FInputConverter::GamepadInputToFKey.Add(MakeTuple(9, Action::Click), EKeys::Gamepad_Special_Right);
-		FInputConverter::GamepadInputToFKey.Add(MakeTuple(10, Action::Click), EKeys::Gamepad_LeftThumbstick);
-		FInputConverter::GamepadInputToFKey.Add(MakeTuple(11, Action::Click), EKeys::Gamepad_RightThumbstick);
-		FInputConverter::GamepadInputToFKey.Add(MakeTuple(12, Action::Click), EKeys::Gamepad_DPad_Up);
-		FInputConverter::GamepadInputToFKey.Add(MakeTuple(13, Action::Click), EKeys::Gamepad_DPad_Down);
-		FInputConverter::GamepadInputToFKey.Add(MakeTuple(14, Action::Click), EKeys::Gamepad_DPad_Left);
-		FInputConverter::GamepadInputToFKey.Add(MakeTuple(15, Action::Click), EKeys::Gamepad_DPad_Right);
-
-		PopulateDefaultCommandHandlers();
-
 		if (UPixelStreaming2PluginSettings::FDelegates* Delegates = UPixelStreaming2PluginSettings::Delegates())
 		{
-			Delegates->OnInputKeyFilterChanged.AddRaw(this, &FPixelStreaming2InputHandler::OnFilteredKeysChanged);
+			OnInputKeyFilterChangedHandle = Delegates->OnInputKeyFilterChanged.AddRaw(this, &FPixelStreaming2InputHandler::OnFilteredKeysChanged);
 		}
+
+		// Register this input handler with the module's input device so that it's ticked
+		FInputDevice::GetInputDevice()->AddInputHandler(this);
 	}
 
 	FPixelStreaming2InputHandler::~FPixelStreaming2InputHandler()
 	{
+		if (UPixelStreaming2PluginSettings::FDelegates* Delegates = UPixelStreaming2PluginSettings::Delegates())
+		{
+			Delegates->OnInputKeyFilterChanged.Remove(OnInputKeyFilterChangedHandle);
+		}
+
+		FInputDevice::GetInputDevice()->RemoveInputHandler(this);
 	}
 
 	TSharedPtr<IPixelStreaming2DataProtocol> FPixelStreaming2InputHandler::GetToStreamerProtocol()
@@ -379,200 +252,151 @@ namespace UE::PixelStreaming2Input
 		// TODO: Implement FFB
 	}
 
-	void FPixelStreaming2InputHandler::HandleOnKeyChar(FMemoryReader Ar)
+	bool FPixelStreaming2InputHandler::OnKeyChar(TCHAR Character)
 	{
-		TPayloadOneParam<TCHAR> Payload(Ar);
-		UE_LOG(LogPixelStreaming2Input, Verbose, TEXT("KEY_PRESSED: Character = '%c'"), Payload.Param1);
+		UE_LOG(LogPixelStreaming2Input, Verbose, TEXT("KEY_PRESSED: Character = '%c'"), Character);
 		// A key char event is never repeated, so set it to false. It's value
 		// ultimately doesn't matter as this paramater isn't used later
-		MessageHandler->OnKeyChar(Payload.Param1, false);
+		return MessageHandler->OnKeyChar(Character, false);
 	}
 
-	void FPixelStreaming2InputHandler::HandleOnKeyDown(FMemoryReader Ar)
+	bool FPixelStreaming2InputHandler::OnKeyDown(FKey Key, bool bIsRepeat)
 	{
-		TPayloadTwoParam<uint8, uint8> Payload(Ar);
+		const uint32* KeyPtr;
+		const uint32* CharacterPtr;
+		FInputKeyManager::Get().GetCodesFromKey(Key, KeyPtr, CharacterPtr);
+		uint32 KeyCode = KeyPtr ? *KeyPtr : 0;
+		uint32 Character = CharacterPtr ? *CharacterPtr : 0;
 
-		bool		bIsRepeat = Payload.Param2 != 0;
-		const FKey* AgnosticKey = JavaScriptKeyCodeToFKey[Payload.Param1];
-		if (FilterKey(*AgnosticKey))
-		{
-			const uint32* KeyPtr;
-			const uint32* CharacterPtr;
-			FInputKeyManager::Get().GetCodesFromKey(*AgnosticKey, KeyPtr, CharacterPtr);
-			uint32 Key = KeyPtr ? *KeyPtr : 0;
-			uint32 Character = CharacterPtr ? *CharacterPtr : 0;
-
-			PixelStreamerApplicationWrapper->UpdateModifierKey(AgnosticKey, true);
-			UE_LOG(LogPixelStreaming2Input, Verbose, TEXT("KEY_DOWN: Key = %d; Character = %d; IsRepeat = %s"), Key, Character, bIsRepeat ? TEXT("True") : TEXT("False"));
-			MessageHandler->OnKeyDown((int32)Key, (int32)Character, bIsRepeat);
-		}
+		PixelStreamerApplicationWrapper->UpdateModifierKey(&Key, true);
+		UE_LOG(LogPixelStreaming2Input, Verbose, TEXT("KEY_DOWN: Key = %d; Character = %d; IsRepeat = %s"), KeyCode, Character, bIsRepeat ? TEXT("True") : TEXT("False"));
+		return MessageHandler->OnKeyDown((int32)KeyCode, (int32)Character, bIsRepeat);
 	}
 
-	void FPixelStreaming2InputHandler::HandleOnKeyUp(FMemoryReader Ar)
+	bool FPixelStreaming2InputHandler::OnKeyUp(FKey Key)
 	{
-		TPayloadOneParam<uint8> Payload(Ar);
-		const FKey*				AgnosticKey = JavaScriptKeyCodeToFKey[Payload.Param1];
-		if (FilterKey(*AgnosticKey))
-		{
-			const uint32* KeyPtr;
-			const uint32* CharacterPtr;
-			FInputKeyManager::Get().GetCodesFromKey(*AgnosticKey, KeyPtr, CharacterPtr);
-			uint32 Key = KeyPtr ? *KeyPtr : 0;
-			uint32 Character = CharacterPtr ? *CharacterPtr : 0;
-
-			PixelStreamerApplicationWrapper->UpdateModifierKey(AgnosticKey, false);
-			UE_LOG(LogPixelStreaming2Input, Verbose, TEXT("KEY_UP: Key = %d; Character = %d"), Key, Character);
-			MessageHandler->OnKeyUp((int32)Key, (int32)Character, false);
-		}
+		const uint32* KeyPtr;
+		const uint32* CharacterPtr;
+		FInputKeyManager::Get().GetCodesFromKey(Key, KeyPtr, CharacterPtr);
+		uint32 KeyCode = KeyPtr ? *KeyPtr : 0;
+		uint32 Character = CharacterPtr ? *CharacterPtr : 0;
+
+		PixelStreamerApplicationWrapper->UpdateModifierKey(&Key, false);
+		UE_LOG(LogPixelStreaming2Input, Verbose, TEXT("KEY_UP: Key = %d; Character = %d"), KeyCode, Character);
+		return MessageHandler->OnKeyUp((int32)KeyCode, (int32)Character, false);
 	}
 
-	void FPixelStreaming2InputHandler::HandleOnTouchStarted(FMemoryReader Ar)
+	bool FPixelStreaming2InputHandler::OnTouchStarted(FIntPoint TouchLocation, int32 TouchIndex, float TouchForce)
 	{
-		TPayloadOneParam<uint8> Payload(Ar);
+		UE_LOG(LogPixelStreaming2Input, Verbose, TEXT("TOUCH_START: TouchIndex = %d; CursorPos = (%d, %d); Force = %.3f"), TouchIndex, static_cast<int>(TouchLocation.X), static_cast<int>(TouchLocation.Y), TouchForce);
 
-		uint8 NumTouches = Payload.Param1;
-		for (uint8 TouchIdx = 0; TouchIdx < NumTouches; TouchIdx++)
-		{
-			//                 PosX    PoxY    IDX   Force  Valid
-			TPayloadFiveParam<uint16, uint16, uint8, uint8, uint8> Touch(Ar);
-			// If Touch is valid
-			if (Touch.Param5 != 0)
-			{
-				//                                                                           convert range from 0,65536 -> 0,1
-				FVector2D	TouchLocation = ConvertFromNormalizedScreenLocation(FVector2D(Touch.Param1 / uint16_MAX, Touch.Param2 / uint16_MAX));
-				const int32 TouchIndex = Touch.Param3;
-				const float TouchForce = Touch.Param4 / 255.0f;
-				UE_LOG(LogPixelStreaming2Input, Verbose, TEXT("TOUCH_START: TouchIndex = %d; Pos = (%d, %d); CursorPos = (%d, %d); Force = %.3f"), TouchIndex, Touch.Param1, Touch.Param2, static_cast<int>(TouchLocation.X), static_cast<int>(TouchLocation.Y), TouchForce);
+		bool bHandled = false;
 
-				if (InputType == EPixelStreaming2InputType::RouteToWidget)
-				{
-					// TouchLocation = TouchLocation - TargetViewport.Pin()->GetCachedGeometry().GetAbsolutePosition();
-					FWidgetPath WidgetPath = FindRoutingMessageWidget(TouchLocation);
-
-					if (WidgetPath.IsValid())
-					{
-						FScopedSwitchWorldHack SwitchWorld(WidgetPath);
-						FPointerEvent		   PointerEvent(0, TouchIndex, TouchLocation, TouchLocation, TouchForce, true);
-						FSlateApplication::Get().RoutePointerDownEvent(WidgetPath, PointerEvent);
-					}
-				}
-				else if (InputType == EPixelStreaming2InputType::RouteToWindow)
-				{
-					if (NumActiveTouches == 0 && !bIsMouseActive)
-					{
-						FSlateApplication::Get().OnCursorSet();
-						// Make sure the application is active.
-						FSlateApplication::Get().ProcessApplicationActivationEvent(true);
+		if (InputType == EPixelStreaming2InputType::RouteToWidget)
+		{
+			// TouchLocation = TouchLocation - TargetViewport.Pin()->GetCachedGeometry().GetAbsolutePosition();
+			FWidgetPath WidgetPath = FindRoutingMessageWidget(TouchLocation);
 
-						FVector2D OldCursorLocation = PixelStreamerApplicationWrapper->WrappedApplication->Cursor->GetPosition();
-						PixelStreamerApplicationWrapper->Cursor->SetPosition(OldCursorLocation.X, OldCursorLocation.Y);
-						FSlateApplication::Get().OverridePlatformApplication(PixelStreamerApplicationWrapper);
-					}
+			if (WidgetPath.IsValid())
+			{
+				FScopedSwitchWorldHack SwitchWorld(WidgetPath);
+				FPointerEvent		   PointerEvent(0, TouchIndex, TouchLocation, TouchLocation, TouchForce, true);
+				bHandled = FSlateApplication::Get().RoutePointerDownEvent(WidgetPath, PointerEvent).IsEventHandled();
+			}
+		}
+		else if (InputType == EPixelStreaming2InputType::RouteToWindow)
+		{
+			if (NumActiveTouches == 0 && !bIsMouseActive)
+			{
+				FSlateApplication::Get().OnCursorSet();
+				// Make sure the application is active.
+				FSlateApplication::Get().ProcessApplicationActivationEvent(true);
 
-					// We must update the user cursor position explicitly before updating the application cursor position
-					// as if there's a delta between them, when the touch event is started it will trigger a move
-					// resulting in a large 'drag' across the screen
-					TSharedPtr<FSlateUser> User = FSlateApplication::Get().GetCursorUser();
-					User->SetCursorPosition(TouchLocation);
-					PixelStreamerApplicationWrapper->Cursor->SetPosition(TouchLocation.X, TouchLocation.Y);
-					PixelStreamerApplicationWrapper->WrappedApplication->Cursor->SetPosition(TouchLocation.X, TouchLocation.Y);
+				FVector2D OldCursorLocation = PixelStreamerApplicationWrapper->WrappedApplication->Cursor->GetPosition();
+				PixelStreamerApplicationWrapper->Cursor->SetPosition(OldCursorLocation.X, OldCursorLocation.Y);
+				FSlateApplication::Get().OverridePlatformApplication(PixelStreamerApplicationWrapper);
+			}
 
-					MessageHandler->OnTouchStarted(PixelStreamerApplicationWrapper->GetWindowUnderCursor(), TouchLocation, TouchForce, TouchIndex, 0); // TODO: ControllerId?
-				}
+			// We must update the user cursor position explicitly before updating the application cursor position
+			// as if there's a delta between them, when the touch event is started it will trigger a move
+			// resulting in a large 'drag' across the screen
+			TSharedPtr<FSlateUser> User = FSlateApplication::Get().GetCursorUser();
+			User->SetCursorPosition(TouchLocation);
+			PixelStreamerApplicationWrapper->Cursor->SetPosition(TouchLocation.X, TouchLocation.Y);
+			PixelStreamerApplicationWrapper->WrappedApplication->Cursor->SetPosition(TouchLocation.X, TouchLocation.Y);
 
-				NumActiveTouches++;
-			}
+			bHandled = MessageHandler->OnTouchStarted(PixelStreamerApplicationWrapper->GetWindowUnderCursor(), TouchLocation, TouchForce, TouchIndex, 0); // TODO: ControllerId?
 		}
 
+		NumActiveTouches++;
+
 		FindFocusedWidget();
+
+		return bHandled;
 	}
 
-	void FPixelStreaming2InputHandler::HandleOnTouchMoved(FMemoryReader Ar)
+	bool FPixelStreaming2InputHandler::OnTouchMoved(FIntPoint TouchLocation, int32 TouchIndex, float TouchForce)
 	{
-		TPayloadOneParam<uint8> Payload(Ar);
-
-		uint8 NumTouches = Payload.Param1;
-		for (uint8 TouchIdx = 0; TouchIdx < NumTouches; TouchIdx++)
-		{
-			//                 PosX    PoxY    IDX   Force  Valid
-			TPayloadFiveParam<uint16, uint16, uint8, uint8, uint8> Touch(Ar);
-			// If Touch is valid
-			if (Touch.Param5 != 0)
-			{
-				//                                                                           convert range from 0,65536 -> 0,1
-				FVector2D	TouchLocation = ConvertFromNormalizedScreenLocation(FVector2D(Touch.Param1 / uint16_MAX, Touch.Param2 / uint16_MAX));
-				const int32 TouchIndex = Touch.Param3;
-				const float TouchForce = Touch.Param4 / 255.0f;
-				UE_LOG(LogPixelStreaming2Input, Verbose, TEXT("TOUCH_MOVE: TouchIndex = %d; Pos = (%d, %d); CursorPos = (%d, %d); Force = %.3f"), TouchIndex, Touch.Param1, Touch.Param2, static_cast<int>(TouchLocation.X), static_cast<int>(TouchLocation.Y), TouchForce);
+		UE_LOG(LogPixelStreaming2Input, Verbose, TEXT("TOUCH_MOVE: TouchIndex = %d; CursorPos = (%d, %d); Force = %.3f"), TouchIndex, static_cast<int>(TouchLocation.X), static_cast<int>(TouchLocation.Y), TouchForce);
 
-				FCachedTouchEvent& TouchEvent = CachedTouchEvents.FindOrAdd(TouchIndex);
-				TouchEvent.Force = TouchForce;
-				TouchEvent.ControllerIndex = 0;
+		FCachedTouchEvent& TouchEvent = CachedTouchEvents.FindOrAdd(TouchIndex);
+		TouchEvent.Force = TouchForce;
+		TouchEvent.ControllerIndex = 0;
 
-				if (InputType == EPixelStreaming2InputType::RouteToWidget)
-				{
-					// TouchLocation = TouchLocation - TargetViewport.Pin()->GetCachedGeometry().GetAbsolutePosition();
-					TouchEvent.Location = TouchLocation;
-					FWidgetPath WidgetPath = FindRoutingMessageWidget(TouchLocation);
+		bool bHandled = false;
 
-					if (WidgetPath.IsValid())
-					{
-						FScopedSwitchWorldHack SwitchWorld(WidgetPath);
-						FPointerEvent		   PointerEvent(0, TouchIndex, TouchLocation, LastTouchLocation, TouchForce, true);
-						FSlateApplication::Get().RoutePointerMoveEvent(WidgetPath, PointerEvent, false);
-					}
+		if (InputType == EPixelStreaming2InputType::RouteToWidget)
+		{
+			// TouchLocation = TouchLocation - TargetViewport.Pin()->GetCachedGeometry().GetAbsolutePosition();
+			TouchEvent.Location = TouchLocation;
+			FWidgetPath WidgetPath = FindRoutingMessageWidget(TouchLocation);
 
-					LastTouchLocation = TouchLocation;
-				}
-				else if (InputType == EPixelStreaming2InputType::RouteToWindow)
-				{
-					TouchEvent.Location = TouchLocation;
-					MessageHandler->OnTouchMoved(TouchEvent.Location, TouchEvent.Force, TouchIndex, TouchEvent.ControllerIndex); // TODO: ControllerId?
-				}
-				TouchIndicesProcessedThisFrame.Add(TouchIndex);
+			if (WidgetPath.IsValid())
+			{
+				FScopedSwitchWorldHack SwitchWorld(WidgetPath);
+				FPointerEvent		   PointerEvent(0, TouchIndex, TouchLocation, LastTouchLocation, TouchForce, true);
+				bHandled = FSlateApplication::Get().RoutePointerMoveEvent(WidgetPath, PointerEvent, false);
 			}
+
+			LastTouchLocation = TouchLocation;
+		}
+		else if (InputType == EPixelStreaming2InputType::RouteToWindow)
+		{
+			TouchEvent.Location = TouchLocation;
+			bHandled = MessageHandler->OnTouchMoved(TouchEvent.Location, TouchEvent.Force, TouchIndex, TouchEvent.ControllerIndex); // TODO: ControllerId?
 		}
+
+		TouchIndicesProcessedThisFrame.Add(TouchIndex);
+
+		return bHandled;
 	}
 
-	void FPixelStreaming2InputHandler::HandleOnTouchEnded(FMemoryReader Ar)
+	bool FPixelStreaming2InputHandler::OnTouchEnded(FIntPoint TouchLocation, int32 TouchIndex)
 	{
-		TPayloadOneParam<uint8> Payload(Ar);
-		uint8					NumTouches = Payload.Param1;
-		for (uint8 TouchIdx = 0; TouchIdx < NumTouches; TouchIdx++)
-		{
-			//                 PosX    PoxY    IDX   Force  Valid
-			TPayloadFiveParam<uint16, uint16, uint8, uint8, uint8> Touch(Ar);
-			// Always allowing the "up" events regardless of in or outside the valid region so
-			// states aren't stuck "down". Might want to uncomment this if it causes other issues.
-			// if(Touch.Param5 != 0)
-			{
-				//                                                                           convert range from 0,65536 -> 0,1
-				FVector2D	TouchLocation = ConvertFromNormalizedScreenLocation(FVector2D(Touch.Param1 / uint16_MAX, Touch.Param2 / uint16_MAX));
-				const int32 TouchIndex = Touch.Param3;
-				UE_LOG(LogPixelStreaming2Input, Verbose, TEXT("TOUCH_END: TouchIndex = %d; Pos = (%d, %d); CursorPos = (%d, %d)"), Touch.Param3, Touch.Param1, Touch.Param2, static_cast<int>(TouchLocation.X), static_cast<int>(TouchLocation.Y));
+		UE_LOG(LogPixelStreaming2Input, Verbose, TEXT("TOUCH_END: TouchIndex = %d; CursorPos = (%d, %d)"), TouchIndex, static_cast<int>(TouchLocation.X), static_cast<int>(TouchLocation.Y));
 
-				if (InputType == EPixelStreaming2InputType::RouteToWidget)
-				{
-					// TouchLocation = TouchLocation - TargetViewport.Pin()->GetCachedGeometry().GetAbsolutePosition();
-					FWidgetPath WidgetPath = FindRoutingMessageWidget(TouchLocation);
+		bool bHandled = false;
 
-					if (WidgetPath.IsValid())
-					{
-						FScopedSwitchWorldHack SwitchWorld(WidgetPath);
-						float				   TouchForce = 0.0f;
-						FPointerEvent		   PointerEvent(0, TouchIndex, TouchLocation, TouchLocation, TouchForce, true);
-						FSlateApplication::Get().RoutePointerUpEvent(WidgetPath, PointerEvent);
-					}
-				}
-				else if (InputType == EPixelStreaming2InputType::RouteToWindow)
-				{
-					MessageHandler->OnTouchEnded(TouchLocation, TouchIndex, 0); // TODO: ControllerId?
-				}
+		if (InputType == EPixelStreaming2InputType::RouteToWidget)
+		{
+			// TouchLocation = TouchLocation - TargetViewport.Pin()->GetCachedGeometry().GetAbsolutePosition();
+			FWidgetPath WidgetPath = FindRoutingMessageWidget(TouchLocation);
 
-				CachedTouchEvents.Remove(TouchIndex);
-				NumActiveTouches = (NumActiveTouches > 0) ? NumActiveTouches - 1 : NumActiveTouches;
+			if (WidgetPath.IsValid())
+			{
+				FScopedSwitchWorldHack SwitchWorld(WidgetPath);
+				float				   TouchForce = 0.0f;
+				FPointerEvent		   PointerEvent(0, TouchIndex, TouchLocation, TouchLocation, TouchForce, true);
+				bHandled = FSlateApplication::Get().RoutePointerUpEvent(WidgetPath, PointerEvent).IsEventHandled();
 			}
 		}
+		else if (InputType == EPixelStreaming2InputType::RouteToWindow)
+		{
+			bHandled = MessageHandler->OnTouchEnded(TouchLocation, TouchIndex, 0); // TODO: ControllerId?
+		}
+
+		CachedTouchEvents.Remove(TouchIndex);
+		NumActiveTouches = (NumActiveTouches > 0) ? NumActiveTouches - 1 : NumActiveTouches;
 
 		// If there's no remaining touches, and there is also no mouse over the player window
 		// then set the platform application back to its default. We need to set it back to default
@@ -584,9 +408,11 @@ namespace UE::PixelStreaming2Input
 			PixelStreamerApplicationWrapper->WrappedApplication->Cursor->SetPosition(OldCursorLocation.X, OldCursorLocation.Y);
 			FSlateApplication::Get().OverridePlatformApplication(PixelStreamerApplicationWrapper->WrappedApplication);
 		}
+
+		return bHandled;
 	}
 
-	void FPixelStreaming2InputHandler::HandleOnControllerConnected(FMemoryReader Ar)
+	uint8 FPixelStreaming2InputHandler::OnControllerConnected()
 	{
 		uint8 NextControllerId = FInputDevice::GetInputDevice()->OnControllerConnected();
 
@@ -604,89 +430,71 @@ namespace UE::PixelStreaming2Input
 		OnSendMessage.Broadcast(EPixelStreaming2FromStreamerMessage::GamepadResponse, FMemoryReader(Data));
 
 		UE_LOG(LogPixelStreaming2Input, Verbose, TEXT("GAMEPAD_CONNECTED: ControllerId = %d"), NextControllerId);
+
+		return NextControllerId;
 	}
 
-	void FPixelStreaming2InputHandler::HandleOnControllerAnalog(FMemoryReader Ar)
+	bool FPixelStreaming2InputHandler::OnControllerAnalog(uint8 ControllerIndex, FKey Key, double AxisValue)
 	{
-		const TPayloadThreeParam<uint8, uint8, double> Payload(Ar);
-
 		FInputDeviceId	DeviceId = INPUTDEVICEID_NONE;
 		FPlatformUserId PlatformUserId = PLATFORMUSERID_NONE;
-		if (!FInputDevice::GetInputDevice()->GetPlatformUserAndDevice(Payload.Param1, DeviceId, PlatformUserId))
-		{
-			return;
-		}
-		double AxisValue = Payload.Param3;
-		FKeyId KeyId = Payload.Param2;
-
-		FKey* AnalogKeyPtr = FInputConverter::GamepadInputToFKey.Find(MakeTuple(KeyId, Action::Axis));
-		if (!AnalogKeyPtr)
+		if (!FInputDevice::GetInputDevice()->GetPlatformUserAndDevice(ControllerIndex, DeviceId, PlatformUserId))
 		{
-			return;
+			return false;
 		}
 
 		FAnalogValue AnalogValue;
 		AnalogValue.Value = AxisValue;
 		// Only send axes values continuously in the case of gamepad triggers
-		AnalogValue.bKeepUnlessZero = (KeyId == 5 || KeyId == 6);
+		AnalogValue.bKeepUnlessZero = (Key == EKeys::Gamepad_LeftTriggerAxis || Key == EKeys::Gamepad_RightTriggerAxis);
 
 		// Overwrite the last data: every tick only process the latest
-		AnalogEventsReceivedThisTick.FindOrAdd(DeviceId).FindOrAdd(AnalogKeyPtr) = AnalogValue;
+		AnalogEventsReceivedThisTick.FindOrAdd(DeviceId).FindOrAdd(Key) = AnalogValue;
+
+		return true;
 	}
 
-	void FPixelStreaming2InputHandler::HandleOnControllerButtonPressed(FMemoryReader Ar)
+	bool FPixelStreaming2InputHandler::OnControllerButtonPressed(uint8 ControllerIndex, FKey Key, bool bIsRepeat)
 	{
-		TPayloadThreeParam<uint8, uint8, uint8> Payload(Ar);
-		FInputDeviceId							DeviceId = INPUTDEVICEID_NONE;
-		FPlatformUserId							PlatformUserId = PLATFORMUSERID_NONE;
-		if (!FInputDevice::GetInputDevice()->GetPlatformUserAndDevice(Payload.Param1, DeviceId, PlatformUserId))
-		{
-			return;
-		}
-		bool  bIsRepeat = Payload.Param3 != 0;
-		FKey* ButtonPtr = FInputConverter::GamepadInputToFKey.Find(MakeTuple(Payload.Param2, Action::Click));
-		if (ButtonPtr == nullptr)
+		FInputDeviceId	DeviceId = INPUTDEVICEID_NONE;
+		FPlatformUserId PlatformUserId = PLATFORMUSERID_NONE;
+		if (!FInputDevice::GetInputDevice()->GetPlatformUserAndDevice(ControllerIndex, DeviceId, PlatformUserId))
 		{
-			return;
+			return false;
 		}
 
-		MessageHandler->OnControllerButtonPressed(ButtonPtr->GetFName(), PlatformUserId, DeviceId, bIsRepeat);
+		UE_LOG(LogPixelStreaming2Input, Verbose, TEXT("GAMEPAD_PRESSED: ControllerId = %d; KeyName = %s; IsRepeat = %s;"), ControllerIndex, *Key.ToString(), bIsRepeat ? TEXT("True") : TEXT("False"));
 
-		UE_LOG(LogPixelStreaming2Input, Verbose, TEXT("GAMEPAD_PRESSED: ControllerId = %d; KeyName = %s; IsRepeat = %s;"), Payload.Param1, *ButtonPtr->ToString(), bIsRepeat ? TEXT("True") : TEXT("False"));
+		return MessageHandler->OnControllerButtonPressed(Key.GetFName(), PlatformUserId, DeviceId, bIsRepeat);
 	}
 
-	void FPixelStreaming2InputHandler::HandleOnControllerButtonReleased(FMemoryReader Ar)
+	bool FPixelStreaming2InputHandler::OnControllerButtonReleased(uint8 ControllerIndex, FKey Key)
 	{
-		TPayloadTwoParam<uint8, uint8> Payload(Ar);
-		FInputDeviceId				   DeviceId = INPUTDEVICEID_NONE;
-		FPlatformUserId				   PlatformUserId = PLATFORMUSERID_NONE;
-		if (!FInputDevice::GetInputDevice()->GetPlatformUserAndDevice(Payload.Param1, DeviceId, PlatformUserId))
-		{
-			return;
-		}
-		FKey* ButtonPtr = FInputConverter::GamepadInputToFKey.Find(MakeTuple(Payload.Param2, Action::Click));
-		if (ButtonPtr == nullptr)
+		FInputDeviceId	DeviceId = INPUTDEVICEID_NONE;
+		FPlatformUserId PlatformUserId = PLATFORMUSERID_NONE;
+		if (!FInputDevice::GetInputDevice()->GetPlatformUserAndDevice(ControllerIndex, DeviceId, PlatformUserId))
 		{
-			return;
+			return false;
 		}
 
-		MessageHandler->OnControllerButtonReleased(ButtonPtr->GetFName(), PlatformUserId, DeviceId, false);
+		UE_LOG(LogPixelStreaming2Input, Verbose, TEXT("GAMEPAD_RELEASED: ControllerId = %d; KeyName = %s;"), ControllerIndex, *Key.ToString());
 
-		UE_LOG(LogPixelStreaming2Input, Verbose, TEXT("GAMEPAD_RELEASED: ControllerId = %d; KeyName = %s;"), Payload.Param1, *ButtonPtr->ToString());
+		return MessageHandler->OnControllerButtonReleased(Key.GetFName(), PlatformUserId, DeviceId, false);
 	}
 
-	void FPixelStreaming2InputHandler::HandleOnControllerDisconnected(FMemoryReader Ar)
+	bool FPixelStreaming2InputHandler::OnControllerDisconnected(uint8 ControllerIndex)
 	{
-		TPayloadOneParam<uint8> Payload(Ar);
-		FInputDevice::GetInputDevice()->OnControllerDisconnected(Payload.Param1);
+		FInputDevice::GetInputDevice()->OnControllerDisconnected(ControllerIndex);
+
+		UE_LOG(LogPixelStreaming2Input, Verbose, TEXT("GAMEPAD_DISCONNECTED: ControllerId = %d"), ControllerIndex);
 
-		UE_LOG(LogPixelStreaming2Input, Verbose, TEXT("GAMEPAD_DISCONNECTED: ControllerId = %d"), Payload.Param1);
+		return true;
 	}
 
 	/**
 	 * Mouse events
 	 */
-	void FPixelStreaming2InputHandler::HandleOnMouseEnter(FMemoryReader Ar)
+	bool FPixelStreaming2InputHandler::OnMouseEnter()
 	{
 		if (NumActiveTouches == 0 && !bIsMouseActive)
 		{
@@ -698,9 +506,10 @@ namespace UE::PixelStreaming2Input
 
 		bIsMouseActive = true;
 		UE_LOG(LogPixelStreaming2Input, Verbose, TEXT("MOUSE_ENTER"));
+		return true;
 	}
 
-	void FPixelStreaming2InputHandler::HandleOnMouseLeave(FMemoryReader Ar)
+	bool FPixelStreaming2InputHandler::OnMouseLeave()
 	{
 		if (NumActiveTouches == 0)
 		{
@@ -709,19 +518,17 @@ namespace UE::PixelStreaming2Input
 		}
 		bIsMouseActive = false;
 		UE_LOG(LogPixelStreaming2Input, Verbose, TEXT("MOUSE_LEAVE"));
+		return true;
 	}
 
-	void FPixelStreaming2InputHandler::HandleOnMouseUp(FMemoryReader Ar)
+	bool FPixelStreaming2InputHandler::OnMouseUp(EMouseButtons::Type Button)
 	{
 		// Ensure we have wrapped the slate application at this point
 		if (!bIsMouseActive)
 		{
-			HandleOnMouseEnter(Ar);
+			OnMouseEnter();
 		}
 
-		TPayloadThreeParam<uint8, uint16, uint16> Payload(Ar);
-
-		EMouseButtons::Type Button = static_cast<EMouseButtons::Type>(Payload.Param1);
 		UE_LOG(LogPixelStreaming2Input, Verbose, TEXT("MOUSE_UP: Button = %d"), Button);
 
 		if (InputType == EPixelStreaming2InputType::RouteToWidget)
@@ -744,31 +551,29 @@ namespace UE::PixelStreaming2Input
 					0,
 					SlateApplication.GetPlatformApplication()->GetModifierKeys());
 
-				SlateApplication.RoutePointerUpEvent(WidgetPath, MouseEvent);
+				return SlateApplication.RoutePointerUpEvent(WidgetPath, MouseEvent).IsEventHandled();
 			}
 		}
 		else if (InputType == EPixelStreaming2InputType::RouteToWindow)
 		{
 			if (Button != EMouseButtons::Type::Invalid)
 			{
-				MessageHandler->OnMouseUp(Button);
+				return MessageHandler->OnMouseUp(Button);
 			}
 		}
+
+		return false;
 	}
 
-	void FPixelStreaming2InputHandler::HandleOnMouseDown(FMemoryReader Ar)
+	bool FPixelStreaming2InputHandler::OnMouseDown(EMouseButtons::Type Button, FIntPoint ScreenLocation)
 	{
 		// Ensure we have wrapped the slate application at this point
 		if (!bIsMouseActive)
 		{
-			HandleOnMouseEnter(Ar);
+			OnMouseEnter();
 		}
 
-		TPayloadThreeParam<uint8, uint16, uint16> Payload(Ar);
-		//                                                                           convert range from 0,65536 -> 0,1
-		FVector2D			ScreenLocation = ConvertFromNormalizedScreenLocation(FVector2D(Payload.Param2 / uint16_MAX, Payload.Param3 / uint16_MAX));
-		EMouseButtons::Type Button = static_cast<EMouseButtons::Type>(Payload.Param1);
-		UE_LOG(LogPixelStreaming2Input, Verbose, TEXT("MOUSE_DOWN: Button = %d; Pos = (%.4f, %.4f)"), Button, ScreenLocation.X, ScreenLocation.Y);
+		UE_LOG(LogPixelStreaming2Input, Verbose, TEXT("MOUSE_DOWN: Button = %d; Pos = (%d, %d)"), Button, ScreenLocation.X, ScreenLocation.Y);
 		// Set cursor pos on mouse down - we may not have moved if this is the very first click
 		FSlateApplication& SlateApplication = FSlateApplication::Get();
 		SlateApplication.OnCursorSet();
@@ -796,22 +601,19 @@ namespace UE::PixelStreaming2Input
 					0,
 					SlateApplication.GetPlatformApplication()->GetModifierKeys());
 
-				SlateApplication.RoutePointerDownEvent(WidgetPath, MouseEvent);
+				return SlateApplication.RoutePointerDownEvent(WidgetPath, MouseEvent).IsEventHandled();
 			}
 		}
 		else if (InputType == EPixelStreaming2InputType::RouteToWindow)
 		{
-			MessageHandler->OnMouseDown(PixelStreamerApplicationWrapper->GetWindowUnderCursor(), Button, ScreenLocation);
+			return MessageHandler->OnMouseDown(PixelStreamerApplicationWrapper->GetWindowUnderCursor(), Button, ScreenLocation);
 		}
+
+		return false;
 	}
 
-	void FPixelStreaming2InputHandler::HandleOnMouseMove(FMemoryReader Ar)
+	bool FPixelStreaming2InputHandler::OnMouseMove(FIntPoint ScreenLocation, FIntPoint Delta)
 	{
-		TPayloadFourParam<uint16, uint16, int16, int16> Payload(Ar);
-		//                                                                           convert range from 0,65536 -> 0,1
-		FIntPoint ScreenLocation = ConvertFromNormalizedScreenLocation(FVector2D(Payload.Param1 / uint16_MAX, Payload.Param2 / uint16_MAX));
-		//                                                                 convert range from -32,768 to 32,767 -> -1,1
-		FIntPoint Delta = ConvertFromNormalizedScreenLocation(FVector2D(Payload.Param3 / int16_MAX, Payload.Param4 / int16_MAX), false);
 		UE_LOG(LogPixelStreaming2Input, Verbose, TEXT("MOUSE_MOVE: Pos = (%d, %d); Delta = (%d, %d)"), ScreenLocation.X, ScreenLocation.Y, Delta.X, Delta.Y);
 		FSlateApplication& SlateApplication = FSlateApplication::Get();
 		SlateApplication.OnCursorSet();
@@ -834,22 +636,20 @@ namespace UE::PixelStreaming2Input
 					SlateApplication.GetPressedMouseButtons(),
 					SlateApplication.GetPlatformApplication()->GetModifierKeys());
 
-				SlateApplication.RoutePointerMoveEvent(WidgetPath, MouseEvent, false);
+				return SlateApplication.RoutePointerMoveEvent(WidgetPath, MouseEvent, false);
 			}
 		}
 		else if (InputType == EPixelStreaming2InputType::RouteToWindow)
 		{
-			MessageHandler->OnRawMouseMove(Delta.X, Delta.Y);
+			return MessageHandler->OnRawMouseMove(Delta.X, Delta.Y);
 		}
+
+		return false;
 	}
 
-	void FPixelStreaming2InputHandler::HandleOnMouseWheel(FMemoryReader Ar)
+	bool FPixelStreaming2InputHandler::OnMouseWheel(FIntPoint ScreenLocation, float MouseWheelDelta)
 	{
-		TPayloadThreeParam<int16, uint16, uint16> Payload(Ar);
-		//                                                                           convert range from 0,65536 -> 0,1
-		FIntPoint	ScreenLocation = ConvertFromNormalizedScreenLocation(FVector2D(Payload.Param2 / uint16_MAX, Payload.Param3 / uint16_MAX));
-		const float SpinFactor = 1 / 120.0f;
-		UE_LOG(LogPixelStreaming2Input, Verbose, TEXT("MOUSE_WHEEL: Delta = %d; Pos = (%d, %d)"), Payload.Param1, ScreenLocation.X, ScreenLocation.Y);
+		UE_LOG(LogPixelStreaming2Input, Verbose, TEXT("MOUSE_WHEEL: Delta = %.4f; Pos = (%d, %d)"), MouseWheelDelta, ScreenLocation.X, ScreenLocation.Y);
 
 		if (InputType == EPixelStreaming2InputType::RouteToWidget)
 		{
@@ -868,25 +668,23 @@ namespace UE::PixelStreaming2Input
 					SlateApplication.GetCursorPos(),
 					SlateApplication.GetPressedMouseButtons(),
 					EKeys::Invalid,
-					Payload.Param1 * SpinFactor,
+					MouseWheelDelta,
 					SlateApplication.GetPlatformApplication()->GetModifierKeys());
 
-				SlateApplication.RouteMouseWheelOrGestureEvent(WidgetPath, MouseEvent, nullptr);
+				return SlateApplication.RouteMouseWheelOrGestureEvent(WidgetPath, MouseEvent, nullptr).IsEventHandled();
 			}
 		}
 		else if (InputType == EPixelStreaming2InputType::RouteToWindow)
 		{
-			MessageHandler->OnMouseWheel(Payload.Param1 * SpinFactor, ScreenLocation);
+			return MessageHandler->OnMouseWheel(MouseWheelDelta, ScreenLocation);
 		}
+
+		return false;
 	}
 
-	void FPixelStreaming2InputHandler::HandleOnMouseDoubleClick(FMemoryReader Ar)
+	bool FPixelStreaming2InputHandler::OnMouseDoubleClick(EMouseButtons::Type Button, FIntPoint ScreenLocation)
 	{
-		TPayloadThreeParam<uint8, uint16, uint16> Payload(Ar);
-		//                                                                           convert range from 0,65536 -> 0,1
-		FVector2D			ScreenLocation = ConvertFromNormalizedScreenLocation(FVector2D(Payload.Param2 / uint16_MAX, Payload.Param3 / uint16_MAX));
-		EMouseButtons::Type Button = static_cast<EMouseButtons::Type>(Payload.Param1);
-		UE_LOG(LogPixelStreaming2Input, Verbose, TEXT("MOUSE_DOWN: Button = %d; Pos = (%.4f, %.4f)"), Button, ScreenLocation.X, ScreenLocation.Y);
+		UE_LOG(LogPixelStreaming2Input, Verbose, TEXT("MOUSE_DOWN: Button = %d; Pos = (%d, %d)"), Button, ScreenLocation.X, ScreenLocation.Y);
 		// Force window focus
 		FSlateApplication& SlateApplication = FSlateApplication::Get();
 		SlateApplication.ProcessApplicationActivationEvent(true);
@@ -910,253 +708,114 @@ namespace UE::PixelStreaming2Input
 					0,
 					SlateApplication.GetPlatformApplication()->GetModifierKeys());
 
-				SlateApplication.RoutePointerDoubleClickEvent(WidgetPath, MouseEvent);
+				return SlateApplication.RoutePointerDoubleClickEvent(WidgetPath, MouseEvent).IsEventHandled();
 			}
 		}
 		else if (InputType == EPixelStreaming2InputType::RouteToWindow)
 		{
-			MessageHandler->OnMouseDoubleClick(PixelStreamerApplicationWrapper->GetWindowUnderCursor(), Button, ScreenLocation);
+			return MessageHandler->OnMouseDoubleClick(PixelStreamerApplicationWrapper->GetWindowUnderCursor(), Button, ScreenLocation);
 		}
-	}
-
-	FMatrix FPixelStreaming2InputHandler::ExtractWebXRMatrix(FMemoryReader& Ar)
-	{
-		FMatrix OutMat;
-		for (int32 Row = 0; Row < 4; ++Row)
 
-		{
-			float Col0 = 0.0f, Col1 = 0.0f, Col2 = 0.0f, Col3 = 0.0f;
-			Ar << Col0 << Col1 << Col2 << Col3;
-			OutMat.M[Row][0] = Col0;
-			OutMat.M[Row][1] = Col1;
-			OutMat.M[Row][2] = Col2;
-			OutMat.M[Row][3] = Col3;
-		}
-		OutMat.DiagnosticCheckNaN();
-		return OutMat;
-	}
-
-	FTransform FPixelStreaming2InputHandler::WebXRMatrixToUETransform(FMatrix Mat)
-	{
-		// Rows and columns are swapped between raw mat and FMat
-		FMatrix UEMatrix = FMatrix(
-			FPlane(Mat.M[0][0], Mat.M[1][0], Mat.M[2][0], Mat.M[3][0]),
-			FPlane(Mat.M[0][1], Mat.M[1][1], Mat.M[2][1], Mat.M[3][1]),
-			FPlane(Mat.M[0][2], Mat.M[1][2], Mat.M[2][2], Mat.M[3][2]),
-			FPlane(Mat.M[0][3], Mat.M[1][3], Mat.M[2][3], Mat.M[3][3]));
-
-		// Extract scale vector and reorder coordinates to be UE coordinate system.
-		FVector ScaleVectorRaw = UEMatrix.GetScaleVector();
-		// Note: We do not invert Z scaling here because we already handle that when we rebuild translation/rot below.
-		FVector ScaleVector = FVector(ScaleVectorRaw.Z, ScaleVectorRaw.X, ScaleVectorRaw.Y);
-
-		// Temporarily remove scaling component as we need rotation axes to be unit length for proper quat conversion
-		UEMatrix.RemoveScaling();
-
-		// Extract & convert translation component to UE coordinate syste,
-		FVector Translation = FVector(-UEMatrix.M[3][2], UEMatrix.M[3][0], UEMatrix.M[3][1]) * 100.0f;
-
-		// Extract & convert rotation component to UE coordinate system
-		FQuat RawRotation(UEMatrix);
-		FQuat Rotation(-RawRotation.Z, RawRotation.X, RawRotation.Y, -RawRotation.W);
-		return FTransform(Rotation, Translation, ScaleVector);
+		return false;
 	}
 
 	/**
 	 * XR Handling
 	 */
-	void FPixelStreaming2InputHandler::HandleOnXREyeViews(FMemoryReader Ar)
+	bool FPixelStreaming2InputHandler::OnXREyeViews(FTransform LeftEyeTransform, FMatrix LeftEyeProjectionMatrix, FTransform RightEyeTransform, FMatrix RightEyeProjectionMatrix, FTransform HMDTransform)
 	{
-		// The `Ar` buffer contains the left eye transform matrix stored as 16 floats
-		FTransform LeftEyeTransform = WebXRMatrixToUETransform(ExtractWebXRMatrix(Ar));
-		// The `Ar` buffer contains the left eye projection matrix stored as 16 floats
-		FMatrix LeftEyeProjectionMatrix = ExtractWebXRMatrix(Ar);
-		// The `Ar` buffer contains the right eye transform matrix stored as 16 floats
-		FTransform RightEyeTransform = WebXRMatrixToUETransform(ExtractWebXRMatrix(Ar));
-		// The `Ar` buffer contains the right eye projection matrix stored as 16 floats
-		FMatrix RightEyeProjectionMatrix = ExtractWebXRMatrix(Ar);
-
-		// The `Ar` buffer contains the right eye projection matrix stored as 16 floats
-		FTransform HMDTransform = WebXRMatrixToUETransform(ExtractWebXRMatrix(Ar));
-
 		if (IPixelStreaming2HMD* HMD = IPixelStreaming2HMDModule::Get().GetPixelStreaming2HMD(); HMD != nullptr)
 		{
 			HMD->SetEyeViews(LeftEyeTransform, LeftEyeProjectionMatrix, RightEyeTransform, RightEyeProjectionMatrix, HMDTransform);
+			return true;
 		}
+
+		return false;
 	}
 
-	void FPixelStreaming2InputHandler::HandleOnXRHMDTransform(FMemoryReader Ar)
+	bool FPixelStreaming2InputHandler::OnXRHMDTransform(FTransform HMDTransform)
 	{
-		// The `Ar` buffer contains the transform matrix stored as 16 floats
-		FTransform HMDTransform = WebXRMatrixToUETransform(ExtractWebXRMatrix(Ar));
-
 		if (IPixelStreaming2HMD* HMD = IPixelStreaming2HMDModule::Get().GetPixelStreaming2HMD(); HMD != nullptr)
 		{
 			HMD->SetTransform(HMDTransform);
+			return true;
 		}
+
+		return false;
 	}
 
-	void FPixelStreaming2InputHandler::HandleOnXRControllerTransform(FMemoryReader Ar)
+	bool FPixelStreaming2InputHandler::OnXRControllerTransform(FTransform ControllerTransform, EControllerHand Handedness)
 	{
-		// The `Ar` buffer contains the transform matrix stored as 16 floats
-		FTransform ControllerTransform = WebXRMatrixToUETransform(ExtractWebXRMatrix(Ar));
-
-		// The `Ar` buffer contains a UInt8 for the handedness
-		EControllerHand Handedness = EControllerHand::Left;
-		Ar << Handedness;
-
 		FPixelStreaming2XRController Controller;
 		Controller.Transform = ControllerTransform;
 		Controller.Handedness = Handedness;
 		XRControllers.Add(Handedness, Controller);
+
+		return true;
 	}
 
-	void FPixelStreaming2InputHandler::HandleOnXRButtonTouched(FMemoryReader Ar)
+	bool FPixelStreaming2InputHandler::OnXRButtonTouched(EControllerHand Handedness, FKey Key, bool bIsRepeat)
 	{
-		TPayloadThreeParam<uint8, uint8, uint8> Payload(Ar);
-		IPlatformInputDeviceMapper&				DeviceMapper = IPlatformInputDeviceMapper::Get();
-		FInputDeviceId							ControllerId = DeviceMapper.GetDefaultInputDevice();
+		IPlatformInputDeviceMapper& DeviceMapper = IPlatformInputDeviceMapper::Get();
+		FInputDeviceId				ControllerId = DeviceMapper.GetDefaultInputDevice();
 
-		XRSystem		System = IPixelStreaming2HMDModule::Get().GetActiveXRSystem();
-		EControllerHand Handedness = static_cast<EControllerHand>(Payload.Param1);
-		uint8			ButtonIdx = Payload.Param2;
-		bool			bIsRepeat = Payload.Param3 != 0;
-
-		FKey* ButtonPtr = FInputConverter::XRInputToFKey.Find(MakeTuple(System, Handedness, ButtonIdx, Action::Touch));
-		if (ButtonPtr == nullptr)
-		{
-			return;
-		}
+		UE_LOG(LogPixelStreaming2Input, Verbose, TEXT("XR_TOUCHED: ControllerId = %d; KeyName = %s; IsRepeat = %s;"), ControllerId.GetId(), *Key.ToString(), bIsRepeat ? TEXT("True") : TEXT("False"));
 
-		MessageHandler->OnControllerButtonPressed(ButtonPtr->GetFName(), PLATFORMUSERID_NONE /* Not used */, ControllerId, bIsRepeat);
-		UE_LOG(LogPixelStreaming2Input, Verbose, TEXT("XR_TOUCHED: ControllerId = %d; KeyName = %s; IsRepeat = %s;"), ControllerId.GetId(), *ButtonPtr->ToString(), bIsRepeat ? TEXT("True") : TEXT("False"));
+		return MessageHandler->OnControllerButtonPressed(Key.GetFName(), PLATFORMUSERID_NONE /* Not used */, ControllerId, bIsRepeat);
 	}
 
-	void FPixelStreaming2InputHandler::HandleOnXRButtonTouchReleased(FMemoryReader Ar)
+	bool FPixelStreaming2InputHandler::OnXRButtonTouchReleased(EControllerHand Handedness, FKey Key)
 	{
-		TPayloadThreeParam<uint8, uint8, uint8> Payload(Ar);
-		IPlatformInputDeviceMapper&				DeviceMapper = IPlatformInputDeviceMapper::Get();
-		FInputDeviceId							ControllerId = DeviceMapper.GetDefaultInputDevice();
-		XRSystem								System = IPixelStreaming2HMDModule::Get().GetActiveXRSystem();
-		EControllerHand							Handedness = static_cast<EControllerHand>(Payload.Param1);
-		uint8									ButtonIdx = Payload.Param2;
-		bool									bIsRepeat = Payload.Param3 != 0;
-
-		FKey* ButtonPtr = FInputConverter::XRInputToFKey.Find(MakeTuple(System, Handedness, ButtonIdx, Action::Touch));
-		if (ButtonPtr == nullptr)
-		{
-			return;
-		}
+		IPlatformInputDeviceMapper& DeviceMapper = IPlatformInputDeviceMapper::Get();
+		FInputDeviceId				ControllerId = DeviceMapper.GetDefaultInputDevice();
 
-		MessageHandler->OnControllerButtonReleased(ButtonPtr->GetFName(), PLATFORMUSERID_NONE /* Not used */, ControllerId, bIsRepeat);
-		UE_LOG(LogPixelStreaming2Input, Verbose, TEXT("XR_TOUCH_RELEASED: ControllerId = %d; KeyName = %s; IsRepeat = %s;"), ControllerId.GetId(), *ButtonPtr->ToString(), bIsRepeat ? TEXT("True") : TEXT("False"));
+		UE_LOG(LogPixelStreaming2Input, Verbose, TEXT("XR_TOUCH_RELEASED: ControllerId = %d; KeyName = %s;"), ControllerId.GetId(), *Key.ToString());
+
+		return MessageHandler->OnControllerButtonReleased(Key.GetFName(), PLATFORMUSERID_NONE /* Not used */, ControllerId, false);
 	}
 
-	void FPixelStreaming2InputHandler::HandleOnXRButtonPressed(FMemoryReader Ar)
+	bool FPixelStreaming2InputHandler::OnXRButtonPressed(EControllerHand Handedness, FKey Key, bool bIsRepeat)
 	{
-		TPayloadFourParam<uint8, uint8, uint8, double> Payload(Ar);
-		IPlatformInputDeviceMapper&					   DeviceMapper = IPlatformInputDeviceMapper::Get();
-		FInputDeviceId								   ControllerId = DeviceMapper.GetDefaultInputDevice();
-
-		XRSystem		System = IPixelStreaming2HMDModule::Get().GetActiveXRSystem();
-		EControllerHand Handedness = static_cast<EControllerHand>(Payload.Param1);
-		uint8			ButtonIdx = Payload.Param2;
-		bool			bIsRepeat = Payload.Param3 != 0;
-
-		FKey* ButtonPtr = FInputConverter::XRInputToFKey.Find(MakeTuple(System, Handedness, ButtonIdx, Action::Click));
+		IPlatformInputDeviceMapper& DeviceMapper = IPlatformInputDeviceMapper::Get();
+		FInputDeviceId				ControllerId = DeviceMapper.GetDefaultInputDevice();
 
-		// Try and see if there is an axis associated with this button (usually the case for triggers)
-		FKey* AxisPtr = FInputConverter::XRInputToFKey.Find(MakeTuple(System, Handedness, ButtonIdx, Action::Axis));
+		UE_LOG(LogPixelStreaming2Input, Verbose, TEXT("XR_PRESSED: ControllerId = %d; KeyName = %s; IsRepeat = %s"), ControllerId.GetId(), *Key.ToString(), bIsRepeat ? TEXT("True") : TEXT("False"));
 
-		// Only send button click if there is no axis associated with this button or this the first press of the axis
-		if (ButtonPtr)
-		{
-			MessageHandler->OnControllerButtonPressed(ButtonPtr->GetFName(), PLATFORMUSERID_NONE /* Not used */, ControllerId, bIsRepeat);
-			UE_LOG(LogPixelStreaming2Input, Verbose, TEXT("XR_PRESSED: ControllerId = %d; KeyName = %s; IsRepeat = %s"), ControllerId.GetId(), *ButtonPtr->ToString(), bIsRepeat ? TEXT("True") : TEXT("False"));
-		}
-
-		// If we have axis associate with this press then set axis value to the button press value
-		if (AxisPtr)
-		{
-			// Trigger axes are not robust to some inputs missing across frames.
-			// So to protect against this case (where PS is too slow to transmit them)
-			// we must set the `bKeepUnlessZero` flag, which allows the input to be continuously passed
-			// until a release is fired.
-			FAnalogValue AnalogValue;
-			AnalogValue.bKeepUnlessZero = true;
-			AnalogValue.Value = Payload.Param4;
-			AnalogValue.bIsRepeat = false;
-			UE_LOG(LogPixelStreaming2Input, Verbose, TEXT("XR_ANALOG: ControllerId = %d; KeyName = %s; IsRepeat = False; AnalogValue = %.4f; [Queued for Tick()]"), ControllerId.GetId(), *AxisPtr->ToString(), AnalogValue.Value);
-			AnalogEventsReceivedThisTick.FindOrAdd(ControllerId).FindOrAdd(AxisPtr) = AnalogValue;
-		}
+		return MessageHandler->OnControllerButtonPressed(Key.GetFName(), PLATFORMUSERID_NONE /* Not used */, ControllerId, bIsRepeat);
 	}
 
-	void FPixelStreaming2InputHandler::HandleOnXRButtonReleased(FMemoryReader Ar)
+	bool FPixelStreaming2InputHandler::OnXRButtonReleased(EControllerHand Handedness, FKey Key)
 	{
-		TPayloadThreeParam<uint8, uint8, uint8> Payload(Ar);
-		IPlatformInputDeviceMapper&				DeviceMapper = IPlatformInputDeviceMapper::Get();
-		FInputDeviceId							ControllerId = DeviceMapper.GetDefaultInputDevice();
-		XRSystem								System = IPixelStreaming2HMDModule::Get().GetActiveXRSystem();
-		EControllerHand							Handedness = static_cast<EControllerHand>(Payload.Param1);
-		uint8									ButtonIdx = Payload.Param2;
-		bool									bIsRepeat = Payload.Param3 != 0;
-
-		// Try and see if there is an axis associated with this button (usually the case for triggers)
-		FKey* AxisPtr = FInputConverter::XRInputToFKey.Find(MakeTuple(System, Handedness, ButtonIdx, Action::Axis));
-
-		// If we have axis associate with this release then set axis value to 0.0
-		if (AxisPtr)
-		{
-			// In the case of an axes release, we should clear any analog value that is being
-			// applied across ticks.
-			FAnalogValue AnalogValue;
-			AnalogValue.bKeepUnlessZero = true;
-			AnalogValue.Value = 0.0;
-			AnalogValue.bIsRepeat = false;
-			UE_LOG(LogPixelStreaming2Input, Verbose, TEXT("XR_ANALOG: ControllerId = %d; KeyName = %s; IsRepeat = False; AnalogValue = %.4f; [Queued for Tick()]"), ControllerId.GetId(), *AxisPtr->ToString(), AnalogValue.Value);
-			AnalogEventsReceivedThisTick.FindOrAdd(ControllerId).FindOrAdd(AxisPtr) = AnalogValue;
-		}
+		IPlatformInputDeviceMapper& DeviceMapper = IPlatformInputDeviceMapper::Get();
+		FInputDeviceId				ControllerId = DeviceMapper.GetDefaultInputDevice();
 
-		// Do the actual release after the analog trigger, as the release can cancel any further inputs
-		FKey* ButtonPtr = FInputConverter::XRInputToFKey.Find(MakeTuple(System, Handedness, ButtonIdx, Action::Click));
-		if (ButtonPtr)
-		{
-			MessageHandler->OnControllerButtonReleased(ButtonPtr->GetFName(), PLATFORMUSERID_NONE /* Not used */, ControllerId, bIsRepeat);
-			UE_LOG(LogPixelStreaming2Input, Verbose, TEXT("XR_RELEASED: ControllerId = %d; KeyName = %s; IsRepeat = %s;"), ControllerId.GetId(), *ButtonPtr->ToString(), bIsRepeat ? TEXT("True") : TEXT("False"));
-		}
+		UE_LOG(LogPixelStreaming2Input, Verbose, TEXT("XR_RELEASED: ControllerId = %d; KeyName = %s;"), ControllerId.GetId(), *Key.ToString());
+
+		return MessageHandler->OnControllerButtonReleased(Key.GetFName(), PLATFORMUSERID_NONE /* Not used */, ControllerId, false);
 	}
 
-	void FPixelStreaming2InputHandler::HandleOnXRAnalog(FMemoryReader Ar)
+	bool FPixelStreaming2InputHandler::OnXRAnalog(EControllerHand Handedness, FKey Key, double AnalogValue)
 	{
-		TPayloadThreeParam<uint8, uint8, double> Payload(Ar);
-		IPlatformInputDeviceMapper&				 DeviceMapper = IPlatformInputDeviceMapper::Get();
-		FInputDeviceId							 ControllerId = DeviceMapper.GetDefaultInputDevice();
-		XRSystem								 System = IPixelStreaming2HMDModule::Get().GetActiveXRSystem();
-		EControllerHand							 Handedness = static_cast<EControllerHand>(Payload.Param1);
-		Action									 InputAction = static_cast<Action>(Payload.Param2 % 2);
-		int										 AxisIndex = Payload.Param2;
-		FKey*									 AnalogKeyPtr = FInputConverter::XRInputToFKey.Find(MakeTuple(System, Handedness, AxisIndex, InputAction));
-		if (AnalogKeyPtr == nullptr)
-		{
-			return;
-		}
+		IPlatformInputDeviceMapper& DeviceMapper = IPlatformInputDeviceMapper::Get();
+		FInputDeviceId				ControllerId = DeviceMapper.GetDefaultInputDevice();
 
 		// This codepath is used for XR joysticks, which seems to be robust to temporary drops in input transmission
 		// so we can safely set `bKeepUnlessZero` to false. However, if we use this for more than joysticks we will have to conditionally set this.
-		FAnalogValue AnalogValue;
-		AnalogValue.bKeepUnlessZero = false;
-		AnalogValue.Value = Payload.Param3;
-		AnalogValue.bIsRepeat = false;
+		FAnalogValue AnalogEvent;
+		AnalogEvent.bKeepUnlessZero = false;
+		AnalogEvent.Value = AnalogValue;
+		AnalogEvent.bIsRepeat = false;
+		AnalogEventsReceivedThisTick.FindOrAdd(ControllerId).FindOrAdd(Key) = AnalogEvent;
 
-		UE_LOG(LogPixelStreaming2Input, VeryVerbose, TEXT("XR_ANALOG: ControllerId = %d; KeyName = %s; IsRepeat = False; AnalogValue = %.4f; [Queued for Tick()]"), ControllerId.GetId(), *AnalogKeyPtr->ToString(), AnalogValue.Value);
-		AnalogEventsReceivedThisTick.FindOrAdd(ControllerId).FindOrAdd(AnalogKeyPtr) = AnalogValue;
+		UE_LOG(LogPixelStreaming2Input, VeryVerbose, TEXT("XR_ANALOG: ControllerId = %d; KeyName = %s; IsRepeat = False; AnalogValue = %.4f; [Queued for Tick()]"), ControllerId.GetId(), *Key.ToString(), AnalogEvent.Value);
+
+		return true;
 	}
 
-	void FPixelStreaming2InputHandler::HandleOnXRSystem(FMemoryReader Ar)
+	bool FPixelStreaming2InputHandler::OnXRSystem(EPixelStreaming2XRSystem System)
 	{
-		uint8 ActiveSystem = (uint8)EPixelStreaming2XRSystem::Unknown;
-		Ar << ActiveSystem;
-		IPixelStreaming2HMDModule::Get().SetActiveXRSystem(static_cast<XRSystem>(ActiveSystem));
+		IPixelStreaming2HMDModule::Get().SetActiveXRSystem(System);
+		return true;
 	}
 
 	void FPixelStreaming2InputHandler::SetCommandHandler(const FString& CommandName, const CommandHandlerFn& Handler)
@@ -1174,124 +833,6 @@ namespace UE::PixelStreaming2Input
 		return !ElevatedCheck || ElevatedCheck(Id);
 	}
 
-	void FPixelStreaming2InputHandler::PopulateDefaultCommandHandlers()
-	{
-		// Execute console commands if passed "ConsoleCommand" and -PixelStreaming2AllowConsoleCommands is on.
-		CommandHandlers.Add(TEXT("ConsoleCommand"), [this](FString SourceId, FString Descriptor, FString ConsoleCommand) {
-			if (!UPixelStreaming2PluginSettings::CVarInputAllowConsoleCommands.GetValueOnAnyThread()
-				|| !IsElevated(SourceId))
-			{
-				return;
-			}
-			GEngine->Exec(GEngine->GetWorld(), *ConsoleCommand);
-		});
-
-		// Change width/height if sent { "Resolution.Width": 1920, "Resolution.Height": 1080 }
-		CommandHandlers.Add(TEXT("Resolution.Width"), [this](FString SourceId, FString Descriptor, FString WidthString) {
-			bool	bSuccess = false;
-			FString HeightString;
-			UE::PixelStreaming2Input::ExtractJsonFromDescriptor(Descriptor, TEXT("Resolution.Height"), HeightString, bSuccess);
-			if (bSuccess && IsElevated(SourceId))
-			{
-				int Width = FCString::Atoi(*WidthString);
-				int Height = FCString::Atoi(*HeightString);
-				if (Width < 1 || Height < 1)
-				{
-					return;
-				}
-
-				FString ChangeResCommand = FString::Printf(TEXT("r.SetRes %dx%d"), Width, Height);
-				GEngine->Exec(GEngine->GetWorld(), *ChangeResCommand);
-			}
-		});
-
-		// Response to "Stat.FPS" by calling "stat fps"
-		CommandHandlers.Add(TEXT("Stat.FPS"), [](FString SourceId, FString Descriptor, FString FPSCommand) {
-			FString StatFPSCommand = FString::Printf(TEXT("stat fps"));
-			GEngine->Exec(GEngine->GetWorld(), *StatFPSCommand);
-		});
-	}
-
-	/**
-	 * Command handling
-	 */
-	void FPixelStreaming2InputHandler::HandleOnCommand(FString SourceId, FMemoryReader Ar)
-	{
-		FString Res;
-		Res.GetCharArray().SetNumUninitialized(Ar.TotalSize() / 2 + 1);
-		Ar.Serialize(Res.GetCharArray().GetData(), Ar.TotalSize());
-
-		FString Descriptor = Res.Mid(MessageHeaderOffset);
-		UE_LOG(LogPixelStreaming2Input, Verbose, TEXT("Command: %s"), *Descriptor);
-
-		// Iterate each command handler and see if the command we got matches any of the bound command names.
-		for (auto& CommandHandlersPair : CommandHandlers)
-		{
-			FString CommandValue;
-			bool	bSuccess = false;
-			UE::PixelStreaming2Input::ExtractJsonFromDescriptor(Descriptor, CommandHandlersPair.Key, CommandValue, bSuccess);
-			if (bSuccess)
-			{
-				// Execute bound command handler with descriptor and parsed command value
-				CommandHandlersPair.Value(SourceId, Descriptor, CommandValue);
-				return;
-			}
-		}
-	}
-
-	/**
-	 * UI Interaction handling
-	 */
-	void FPixelStreaming2InputHandler::HandleUIInteraction(FMemoryReader Ar)
-	{
-		// Actual implementation is in the Pixel Streaming module
-	}
-
-	/**
-	 * Textbox Entry handling
-	 */
-	void FPixelStreaming2InputHandler::HandleOnTextboxEntry(FMemoryReader Ar)
-	{
-		FString Res;
-		Res.GetCharArray().SetNumUninitialized(Ar.TotalSize() / 2 + 1);
-		Ar.Serialize(Res.GetCharArray().GetData(), Ar.TotalSize());
-		FString Text = Res.Mid(1);
-
-		FSlateApplication::Get().ForEachUser([this, Text](FSlateUser& User) {
-			TSharedPtr<SWidget> FocusedWidget = User.GetFocusedWidget();
-
-			bool bIsEditableTextType = FocusedWidget->GetType() == TEXT("SEditableText");
-			bool bIsMultiLineEditableTextType = FocusedWidget->GetType() == TEXT("SMultiLineEditableText");
-			bool bEditable = FocusedWidget && (bIsEditableTextType || bIsMultiLineEditableTextType);
-
-			if (bEditable)
-			{
-				if (bIsEditableTextType)
-				{
-					SEditableText* TextBox = static_cast<SEditableText*>(FocusedWidget.Get());
-					TextBox->SetText(FText::FromString(Text));
-				}
-				else if (bIsMultiLineEditableTextType)
-				{
-					SMultiLineEditableTextBox* TextBox = static_cast<SMultiLineEditableTextBox*>(FocusedWidget.Get());
-					TextBox->SetText(FText::FromString(Text));
-				}
-
-				// We need to manually trigger an Enter key press so that the OnTextCommitted delegate gets fired
-				const uint32* KeyPtr = nullptr;
-				const uint32* CharacterPtr = nullptr;
-				FInputKeyManager::Get().GetCodesFromKey(EKeys::Enter, KeyPtr, CharacterPtr);
-				uint32 Key = KeyPtr ? *KeyPtr : 0;
-				uint32 Character = CharacterPtr ? *CharacterPtr : 0;
-				if (Key != 0 || Character != 0)
-				{
-					MessageHandler->OnKeyDown((int32)Key, (int32)Character, false);
-					MessageHandler->OnKeyUp((int32)Key, (int32)Character, false);
-				}
-			}
-		});
-	}
-
 	FIntPoint FPixelStreaming2InputHandler::ConvertFromNormalizedScreenLocation(const FVector2D& ScreenLocation, bool bIncludeOffset)
 	{
 		FIntPoint OutVector((int32)ScreenLocation.X, (int32)ScreenLocation.Y);
@@ -1378,18 +919,17 @@ namespace UE::PixelStreaming2Input
 			for (auto FKeyIt = AnalogInputIt->Value.CreateIterator(); FKeyIt; ++FKeyIt)
 			{
 				const FInputDeviceId& ControllerId = AnalogInputIt->Key;
-				FKey*				  Key = FKeyIt->Key;
+				FKey				  Key = FKeyIt->Key;
 				FAnalogValue		  AnalogValue = FKeyIt->Value;
 				bool				  bIsRepeat = AnalogValue.bIsRepeat;
-				if (!Key)
-				{
-					return;
-				}
 
 				// Check if this gamepad event is specific to a normal gamepad and not an xr gamepad
-				TArray<FKey> GamepadKeys;
-				FInputConverter::GamepadInputToFKey.GenerateValueArray(GamepadKeys);
-				if (GamepadKeys.Find(*Key) != INDEX_NONE)
+				if (Key == EKeys::Gamepad_LeftX ||			 //
+					Key == EKeys::Gamepad_LeftY ||			 //
+					Key == EKeys::Gamepad_RightX ||			 //
+					Key == EKeys::Gamepad_RightY ||			 //
+					Key == EKeys::Gamepad_LeftTriggerAxis || //
+					Key == EKeys::Gamepad_RightTriggerAxis)
 				{
 					// This is a gamepad key. We need to check that the gamepad hasn't been disconnected for we've been able to process this analog event
 					uint8 ControllerIndex;
@@ -1402,8 +942,8 @@ namespace UE::PixelStreaming2Input
 				}
 
 				// Pass an analog input along the engine's input processing system
-				bool bHandled = MessageHandler->OnControllerAnalog(Key->GetFName(), PLATFORMUSERID_NONE /* Not used */, ControllerId, AnalogValue.Value);
-				UE_LOG(LogPixelStreaming2Input, VeryVerbose, TEXT("TICKED ANALOG Input: ControllerId = %d; KeyName = %s; IsRepeat = %s; AnalogValue = %.4f; Handled = %s; [Queued for Tick()]"), ControllerId.GetId(), *Key->ToString(), bIsRepeat ? TEXT("True") : TEXT("False"), AnalogValue.Value, bHandled ? TEXT("True") : TEXT("False"));
+				bool bHandled = MessageHandler->OnControllerAnalog(Key.GetFName(), PLATFORMUSERID_NONE /* Not used */, ControllerId, AnalogValue.Value);
+				UE_LOG(LogPixelStreaming2Input, VeryVerbose, TEXT("TICKED ANALOG Input: ControllerId = %d; KeyName = %s; IsRepeat = %s; AnalogValue = %.4f; Handled = %s; [Queued for Tick()]"), ControllerId.GetId(), *Key.ToString(), bIsRepeat ? TEXT("True") : TEXT("False"), AnalogValue.Value, bHandled ? TEXT("True") : TEXT("False"));
 
 				// Remove current analog key unless it has the special `bKeepUnlessZero` flag set.
 				// This flag is used to continuously apply input values across ticks because
diff --git a/Engine/Plugins/Media/PixelStreaming2/Source/PixelStreaming2Input/Private/InputHandler.h b/Engine/Plugins/Media/PixelStreaming2/Source/PixelStreaming2Input/Private/InputHandler.h
deleted file mode 100644
diff --git a/Engine/Plugins/Media/PixelStreaming2/Source/PixelStreaming2Input/Private/InputStructures.h b/Engine/Plugins/Media/PixelStreaming2/Source/PixelStreaming2Input/Private/InputStructures.h
deleted file mode 100644
diff --git a/Engine/Plugins/Media/PixelStreaming2/Source/PixelStreaming2Input/Private/JavaScriptKeyCodes.inl b/Engine/Plugins/Media/PixelStreaming2/Source/PixelStreaming2Input/Private/JavaScriptKeyCodes.inl
deleted file mode 100644
diff --git a/Engine/Plugins/Media/PixelStreaming2/Source/PixelStreaming2Input/Private/PixelStreaming2InputModule.cpp b/Engine/Plugins/Media/PixelStreaming2/Source/PixelStreaming2Input/Private/PixelStreaming2InputModule.cpp
--- a/Engine/Plugins/Media/PixelStreaming2/Source/PixelStreaming2Input/Private/PixelStreaming2InputModule.cpp
+++ b/Engine/Plugins/Media/PixelStreaming2/Source/PixelStreaming2Input/Private/PixelStreaming2InputModule.cpp
@@ -9,7 +9,6 @@
 
 namespace UE::PixelStreaming2Input
 {
-
 	void FPixelStreaming2InputModule::StartupModule()
 	{
 		if (!FSlateApplication::IsInitialized())
@@ -29,14 +28,7 @@ namespace UE::PixelStreaming2Input
 
 	TSharedPtr<IPixelStreaming2InputHandler> FPixelStreaming2InputModule::CreateInputHandler()
 	{
-		TSharedPtr<FPixelStreaming2ApplicationWrapper> PixelStreamerApplicationWrapper = MakeShareable(new FPixelStreaming2ApplicationWrapper(FSlateApplication::Get().GetPlatformApplication()));
-		TSharedPtr<FGenericApplicationMessageHandler>		 BaseHandler = FSlateApplication::Get().GetPlatformApplication()->GetMessageHandler();
-		TSharedPtr<IPixelStreaming2InputHandler>		 InputHandler = MakeShared<FPixelStreaming2InputHandler>(PixelStreamerApplicationWrapper, BaseHandler);
-
-		// Add this input handler to the input device's array of handlers. This ensures that it's ticked
-		InputDevice->AddInputHandler(InputHandler);
-
-		return InputHandler;
+		return MakeShared<FPixelStreaming2InputHandler>();
 	}
 
 	TSharedPtr<IInputDevice> FPixelStreaming2InputModule::CreateInputDevice(const TSharedRef<FGenericApplicationMessageHandler>& InMessageHandler)
diff --git a/Engine/Plugins/Media/PixelStreaming2/Source/PixelStreaming2Input/Public/IPixelStreaming2InputHandler.h b/Engine/Plugins/Media/PixelStreaming2/Source/PixelStreaming2Input/Public/IPixelStreaming2InputHandler.h
--- a/Engine/Plugins/Media/PixelStreaming2/Source/PixelStreaming2Input/Public/IPixelStreaming2InputHandler.h
+++ b/Engine/Plugins/Media/PixelStreaming2/Source/PixelStreaming2Input/Public/IPixelStreaming2InputHandler.h
@@ -4,12 +4,13 @@
 
 #include "GenericPlatform/GenericPlatformMisc.h"
 #include "IInputDevice.h"
+#include "IPixelStreaming2DataProtocol.h"
+#include "PixelStreaming2HMDEnums.h"
+#include "PixelStreaming2InputEnums.h"
+#include "Serialization/MemoryReader.h"
+#include "Templates/SharedPointer.h"
 #include "Widgets/SViewport.h"
 #include "Widgets/SWindow.h"
-#include "Templates/SharedPointer.h"
-#include "Serialization/MemoryReader.h"
-#include "PixelStreaming2InputEnums.h"
-#include "IPixelStreaming2DataProtocol.h"
 
 /**
  * The IPixelStreaming2InputHandler, used to handle input from a remote peer and pass it to UE accordingly. Setting the target viewport allows for
@@ -147,4 +148,269 @@ public:
 	 * Examples include when a virtual gamepad controller is "connected" and given a unique id.
 	 */
 	FOnSendMessage OnSendMessage;
+
+	/**
+	 * @brief Notify the input handler of a character event.
+	 *
+	 * @param Character The character
+	 *
+	 * @return true if the event was handled. false otherwise
+	 */
+	virtual bool OnKeyChar(TCHAR Character) = 0;
+
+	/**
+	 * @brief Notify the input handler of a key down event.
+	 *
+	 * @param Key The key
+	 * @param bIsRepeat Is this a repeat of a previous Key down (ie has the key been held down)
+	 *
+	 * @return true if the event was handled. false otherwise
+	 */
+	virtual bool OnKeyDown(FKey Key, bool bIsRepeat) = 0;
+
+	/**
+	 * @brief Notify the input handler of a key up event.
+	 *
+	 * @param Key The key
+	 *
+	 * @return true if the event was handled. false otherwise
+	 */
+	virtual bool OnKeyUp(FKey Key) = 0;
+
+	/**
+	 * @brief Notify the input handler of the mouse entering the application (typically the browser window).
+	 *
+	 * @return true if the event was handled. false otherwise
+	 */
+	virtual bool OnMouseEnter() = 0;
+
+	/**
+	 * @brief Notify the input handler of the mouse leaving the application (typically the browser window).
+	 *
+	 * @return true if the event was handled. false otherwise
+	 */
+	virtual bool OnMouseLeave() = 0;
+
+	/**
+	 * @brief Notify the input handler of a mouse down event.
+	 *
+	 * @param Button The mouse button
+	 * @param ScreenPosition The position of the mouse event in UE screen space
+	 *
+	 * @return true if the event was handled. false otherwise
+	 */
+	virtual bool OnMouseDown(EMouseButtons::Type Button, FIntPoint ScreenPosition) = 0;
+
+	/**
+	 * @brief Notify the input handler of a mouse up event.
+	 *
+	 * @param Button The mouse button
+	 *
+	 * @return true if the event was handled. false otherwise
+	 */
+	virtual bool OnMouseUp(EMouseButtons::Type Button) = 0;
+
+	/**
+	 * @brief Notify the input handler of a mouse move event.
+	 *
+	 * @param ScreenPosition The starting position of the mouse event in UE screen space
+	 * @param Delta The amount the mouse moves in this event in UE screen space
+	 *
+	 * @return true if the event was handled. false otherwise
+	 */
+	virtual bool OnMouseMove(FIntPoint ScreenPosition, FIntPoint Delta) = 0;
+
+	/**
+	 * @brief Notify the input handler of a mouse wheel event.
+	 *
+	 * @param ScreenPosition The position of the mouse event in UE screen space
+	 * @param MouseWheelDelta The amount the mouse wheel moved
+	 *
+	 * @return true if the event was handled. false otherwise
+	 */
+	virtual bool OnMouseWheel(FIntPoint ScreenPosition, float MouseWheelDelta) = 0;
+
+	/**
+	 * @brief Notify the input handler of a mouse double click event.
+	 *
+	 * @param Button The mouse button
+	 * @param ScreenPosition The position of the mouse event in UE screen space
+	 *
+	 * @return true if the event was handled. false otherwise
+	 */
+	virtual bool OnMouseDoubleClick(EMouseButtons::Type Button, FIntPoint ScreenPosition) = 0;
+
+	/**
+	 * @brief Notify the input handler of a touch start event.
+	 *
+	 * @param TouchPosition The position of the touch event in UE screen space
+	 * @param TouchIndex The index of the finger that started touching
+	 * @param Force The force of the touch in the range 0 (no force) to 1 (maximum force)
+	 *
+	 * @return true if the event was handled. false otherwise
+	 */
+	virtual bool OnTouchStarted(FIntPoint TouchPosition, int32 TouchIndex, float Force) = 0;
+
+	/**
+	 * @brief Notify the input handler of a touch moved event.
+	 *
+	 * @param TouchPosition The position of the touch event after the move has completed in UE screen space
+	 * @param TouchIndex The index of the finger that moved
+	 * @param Force The force of the touch in the range 0 (no force) to 1 (maximum force)
+	 *
+	 * @return true if the event was handled. false otherwise
+	 */
+	virtual bool OnTouchMoved(FIntPoint TouchPosition, int32 TouchIndex, float Force) = 0;
+
+	/**
+	 * @brief Notify the input handler of a touch ended event.
+	 *
+	 * @param TouchPosition The position of the touch event in UE screen space
+	 * @param TouchIndex The index of the finger that stopped touching
+	 *
+	 * @return true if the event was handled. false otherwise
+	 */
+	virtual bool OnTouchEnded(FIntPoint TouchPosition, int32 TouchIndex) = 0;
+
+	/**
+	 * @brief Notify the input handler of a controller connecting.
+	 *
+	 * @return the index to use for this new controller
+	 */
+	virtual uint8 OnControllerConnected() = 0;
+
+	/**
+	 * @brief Notify the input handler of a controller analog event.
+	 *
+	 * @param ControlledIndex The index of the controller received from OnControllerConnected
+	 * @param Axis The axis FKey
+	 * @param AxisValue The value of the axis in the range -1 to 1
+	 *
+	 * @return true if the event was handled. false otherwise
+	 */
+	virtual bool OnControllerAnalog(uint8 ControllerIndex, FKey Axis, double AxisValue) = 0;
+
+	/**
+	 * @brief Notify the input handler of a controller button press event.
+	 *
+	 * @param ControlledIndex The index of the controller received from OnControllerConnected
+	 * @param Key The button FKey
+	 * @param bIsRepeat Is this a repeat of a previous button press (ie has the button been held down)
+	 *
+	 * @return true if the event was handled. false otherwise
+	 */
+	virtual bool OnControllerButtonPressed(uint8 ControllerIndex, FKey Key, bool bIsRepeat) = 0;
+
+	/**
+	 * @brief Notify the input handler of a controller button release event.
+	 *
+	 * @param ControlledIndex The index of the controller received from OnControllerConnected
+	 * @param Key The button FKey
+	 *
+	 * @return true if the event was handled. false otherwise
+	 */
+	virtual bool OnControllerButtonReleased(uint8 ControllerIndex, FKey Key) = 0;
+
+	/**
+	 * @brief Notify the input handler of a controller disconnecting.
+	 *
+	 * @param ControlledIndex The index of the controller received from OnControllerConnected
+	 *
+	 * @return true if the event was handled. false otherwise
+	 */
+	virtual bool OnControllerDisconnected(uint8 ControllerIndex) = 0;
+
+	/**
+	 * @brief Notify the input handler of receiving the Eye Views to use with XR streaming (usually received oncne at the start of a stream)
+	 *
+	 * @param LeftEyeTransform The transform of the left eye
+	 * @param LeftEyeProjectionMatrix The projection matrix of the left eye
+	 * @param RightEyeTransform The transform of the right eye
+	 * @param RightEyeProjectionMatrix The projection matrix of the right eye
+	 * @param HMDTransform The transform of the HMD itself
+	 *
+	 * @return true if the event was handled. false otherwise
+	 */
+	virtual bool OnXREyeViews(FTransform LeftEyeTransform, FMatrix LeftEyeProjectionMatrix, FTransform RightEyeTransform, FMatrix RightEyeProjectionMatrix, FTransform HMDTransform) = 0;
+
+	/**
+	 * @brief Notify the input handler of receiving the transform of the HMD to use with XR streaming (received once per frame that the HMD displays).
+	 *
+	 * @param HMDTransform The transform of the HMD itself
+	 *
+	 * @return true if the event was handled. false otherwise
+	 */
+	virtual bool OnXRHMDTransform(FTransform HMDTransform) = 0;
+
+	/**
+	 * @brief Notify the input handler of receiving the transform of an XR controller to use with XR streaming (received once per frame that the HMD displays).
+	 *
+	 * @param ControllerTransform The transform of the controller
+	 * @param Handedness Specifies the hand this transform applies to
+	 *
+	 * @return true if the event was handled. false otherwise
+	 */
+	virtual bool OnXRControllerTransform(FTransform ControllerTransform, EControllerHand Handedness) = 0;
+
+	/**
+	 * @brief Notify the input handler of receiving a XR controller button touch event to use with XR streaming
+	 *
+	 * @param Handedness Specifies the hand this button touch event applies to
+	 * @param Key The button FKey
+	 * @param bIsRepeat Is this a repeat of a previous button press (ie has the button been held down)
+	 *
+	 * @return true if the event was handled. false otherwise
+	 */
+	virtual bool OnXRButtonTouched(EControllerHand Handedness, FKey Key, bool bIsRepeat) = 0;
+
+	/**
+	 * @brief Notify the input handler of receiving a XR controller button touch release event to use with XR streaming
+	 *
+	 * @param Handedness Specifies the hand this button touch event applies to
+	 * @param Key The button FKey
+	 *
+	 * @return true if the event was handled. false otherwise
+	 */
+	virtual bool OnXRButtonTouchReleased(EControllerHand Handedness, FKey Key) = 0;
+
+	/**
+	 * @brief Notify the input handler of receiving a XR controller button press event to use with XR streaming
+	 *
+	 * @param Handedness Specifies the hand this button press event applies to
+	 * @param Key The button FKey
+	 * @param bIsRepeat Is this a repeat of a previous button press (ie has the button been held down)
+	 *
+	 * @return true if the event was handled. false otherwise
+	 */
+	virtual bool OnXRButtonPressed(EControllerHand Handedness, FKey Key, bool bIsRepeat) = 0;
+
+	/**
+	 * @brief Notify the input handler of receiving a XR controller button release event to use with XR streaming
+	 *
+	 * @param Handedness Specifies the hand this button release event applies to
+	 * @param Key The button FKey
+	 *
+	 * @return true if the event was handled. false otherwise
+	 */
+	virtual bool OnXRButtonReleased(EControllerHand Handedness, FKey Key) = 0;
+
+	/**
+	 * @brief Notify the input handler of receiving a XR controller button release event to use with XR streaming
+	 *
+	 * @param Handedness Specifies the hand this button release event applies to
+	 * @param Key The axis FKey
+	 * @param AnalogValue The value of the axis in the range -1 to 1
+	 *
+	 * @return true if the event was handled. false otherwise
+	 */
+	virtual bool OnXRAnalog(EControllerHand Handedness, FKey Key, double AnalogValue) = 0;
+
+	/**
+	 * @brief Notify the input handler of the headset system connected to the frontend to use with XR streaming
+	 *
+	 * @param System The system
+	 *
+	 * @return true if the event was handled. false otherwise
+	 */
+	virtual bool OnXRSystem(EPixelStreaming2XRSystem System) = 0;
 };
\ No newline at end of file
diff --git a/Engine/Plugins/Media/PixelStreaming2/Source/PixelStreaming2RTC/PixelStreaming2RTC.Build.cs b/Engine/Plugins/Media/PixelStreaming2/Source/PixelStreaming2RTC/PixelStreaming2RTC.Build.cs
--- a/Engine/Plugins/Media/PixelStreaming2/Source/PixelStreaming2RTC/PixelStreaming2RTC.Build.cs
+++ b/Engine/Plugins/Media/PixelStreaming2/Source/PixelStreaming2RTC/PixelStreaming2RTC.Build.cs
@@ -103,7 +103,8 @@ namespace UnrealBuildTool.Rules
                 "EpicRtc",
                 "AudioPlatformConfiguration",
                 "LibVpx",
-                "PixelStreaming2Settings"
+                "PixelStreaming2Settings",
+                "XRBase"
             });
 
             PrivateDefinitions.Add("PIXELSTREAMING2_DUMP_ENCODING=0");
diff --git a/Engine/Plugins/Media/PixelStreaming2/Source/PixelStreaming2RTC/Private/EpicRtcStreamer.cpp b/Engine/Plugins/Media/PixelStreaming2/Source/PixelStreaming2RTC/Private/EpicRtcStreamer.cpp
--- a/Engine/Plugins/Media/PixelStreaming2/Source/PixelStreaming2RTC/Private/EpicRtcStreamer.cpp
+++ b/Engine/Plugins/Media/PixelStreaming2/Source/PixelStreaming2RTC/Private/EpicRtcStreamer.cpp
@@ -30,12 +30,13 @@
 #include "UtilsCodecs.h"
 #include "VideoProducerBackBuffer.h"
 #include "VideoProducerMediaCapture.h"
+#include "RTCInputHandler.h"
 
 namespace UE::PixelStreaming2
 {
 	FEpicRtcStreamer::FEpicRtcStreamer(const FString& InStreamerId, TRefCountPtr<EpicRtcConferenceInterface> Conference)
 		: StreamerId(InStreamerId)
-		, InputHandler(IPixelStreaming2InputModule::Get().CreateInputHandler())
+		, InputHandler(FRTCInputHandler::Create())
 		, Players(new TThreadSafeMap<FString, FPlayerContext>())
 		, VideoCapturer(FEpicRtcVideoCapturer::Create())
 		, VideoSourceGroup(FVideoSourceGroup::Create(VideoCapturer))
diff --git a/Engine/Plugins/Media/PixelStreaming2/Source/PixelStreaming2RTC/Private/RTCInputHandler.cpp b/Engine/Plugins/Media/PixelStreaming2/Source/PixelStreaming2RTC/Private/RTCInputHandler.cpp
new file mode 100644
--- /dev/null
+++ b/Engine/Plugins/Media/PixelStreaming2/Source/PixelStreaming2RTC/Private/RTCInputHandler.cpp
@@ -0,0 +1,703 @@
+// Copyright Epic Games, Inc. All Rights Reserved.
+
+#include "RTCInputHandler.h"
+
+#include "DefaultDataProtocol.h"
+#include "Engine/Engine.h"
+#include "Framework/Application/SlateApplication.h"
+#include "Framework/Application/SlateUser.h"
+#include "InputStructures.h"
+#include "IPixelStreaming2HMDModule.h"
+#include "JavaScriptKeyCodes.inl"
+#include "Logging.h"
+#include "PixelStreaming2PluginSettings.h"
+#include "PixelStreaming2Utils.h"
+#include "Widgets/Input/SEditableText.h"
+#include "Widgets/Input/SMultiLineEditableTextBox.h"
+
+namespace UE::PixelStreaming2
+{
+	using namespace UE::PixelStreaming2Input;
+
+	TSharedPtr<FRTCInputHandler> FRTCInputHandler::Create()
+	{
+		return MakeShareable<FRTCInputHandler>(new FRTCInputHandler());
+	}
+
+	FRTCInputHandler::FRTCInputHandler()
+	{
+		// RTC module uses the default protocol
+		ToStreamerProtocol = UE::PixelStreaming2Input::GetDefaultToStreamerProtocol();
+		FromStreamerProtocol = UE::PixelStreaming2Input::GetDefaultFromStreamerProtocol();
+
+		RegisterMessageHandler("TouchStart", [this](FString SourceId, FMemoryReader Ar) { HandleOnTouchStarted(Ar); });
+		RegisterMessageHandler("TouchMove", [this](FString SourceId, FMemoryReader Ar) { HandleOnTouchMoved(Ar); });
+		RegisterMessageHandler("TouchEnd", [this](FString SourceId, FMemoryReader Ar) { HandleOnTouchEnded(Ar); });
+
+		RegisterMessageHandler("KeyPress", [this](FString SourceId, FMemoryReader Ar) { HandleOnKeyChar(Ar); });
+		RegisterMessageHandler("KeyUp", [this](FString SourceId, FMemoryReader Ar) { HandleOnKeyUp(Ar); });
+		RegisterMessageHandler("KeyDown", [this](FString SourceId, FMemoryReader Ar) { HandleOnKeyDown(Ar); });
+
+		RegisterMessageHandler("MouseEnter", [this](FString SourceId, FMemoryReader Ar) { HandleOnMouseEnter(Ar); });
+		RegisterMessageHandler("MouseLeave", [this](FString SourceId, FMemoryReader Ar) { HandleOnMouseLeave(Ar); });
+		RegisterMessageHandler("MouseUp", [this](FString SourceId, FMemoryReader Ar) { HandleOnMouseUp(Ar); });
+		RegisterMessageHandler("MouseDown", [this](FString SourceId, FMemoryReader Ar) { HandleOnMouseDown(Ar); });
+		RegisterMessageHandler("MouseMove", [this](FString SourceId, FMemoryReader Ar) { HandleOnMouseMove(Ar); });
+		RegisterMessageHandler("MouseWheel", [this](FString SourceId, FMemoryReader Ar) { HandleOnMouseWheel(Ar); });
+		RegisterMessageHandler("MouseDouble", [this](FString SourceId, FMemoryReader Ar) { HandleOnMouseDoubleClick(Ar); });
+
+		RegisterMessageHandler("GamepadConnected", [this](FString SourceId, FMemoryReader Ar) { HandleOnControllerConnected(Ar); });
+		RegisterMessageHandler("GamepadAnalog", [this](FString SourceId, FMemoryReader Ar) { HandleOnControllerAnalog(Ar); });
+		RegisterMessageHandler("GamepadButtonPressed", [this](FString SourceId, FMemoryReader Ar) { HandleOnControllerButtonPressed(Ar); });
+		RegisterMessageHandler("GamepadButtonReleased", [this](FString SourceId, FMemoryReader Ar) { HandleOnControllerButtonReleased(Ar); });
+		RegisterMessageHandler("GamepadDisconnected", [this](FString SourceId, FMemoryReader Ar) { HandleOnControllerDisconnected(Ar); });
+
+		RegisterMessageHandler("XREyeViews", [this](FString SourceId, FMemoryReader Ar) { HandleOnXREyeViews(Ar); });
+		RegisterMessageHandler("XRHMDTransform", [this](FString SourceId, FMemoryReader Ar) { HandleOnXRHMDTransform(Ar); });
+		RegisterMessageHandler("XRControllerTransform", [this](FString SourceId, FMemoryReader Ar) { HandleOnXRControllerTransform(Ar); });
+		RegisterMessageHandler("XRButtonPressed", [this](FString SourceId, FMemoryReader Ar) { HandleOnXRButtonPressed(Ar); });
+		RegisterMessageHandler("XRButtonTouched", [this](FString SourceId, FMemoryReader Ar) { HandleOnXRButtonTouched(Ar); });
+		RegisterMessageHandler("XRButtonTouchReleased", [this](FString SourceId, FMemoryReader Ar) { HandleOnXRButtonTouchReleased(Ar); });
+		RegisterMessageHandler("XRButtonReleased", [this](FString SourceId, FMemoryReader Ar) { HandleOnXRButtonReleased(Ar); });
+		RegisterMessageHandler("XRAnalog", [this](FString SourceId, FMemoryReader Ar) { HandleOnXRAnalog(Ar); });
+		RegisterMessageHandler("EPixelStreaming2XRSystem", [this](FString SourceId, FMemoryReader Ar) { HandleOnXRSystem(Ar); });
+
+		RegisterMessageHandler("Command", [this](FString SourceId, FMemoryReader Ar) { HandleOnCommand(SourceId, Ar); });
+		RegisterMessageHandler("UIInteraction", [this](FString SourceId, FMemoryReader Ar) { HandleUIInteraction(Ar); });
+		RegisterMessageHandler("TextboxEntry", [this](FString SourceId, FMemoryReader Ar) { HandleOnTextboxEntry(Ar); });
+
+		// RequestQualityControl has been removed. We no-op this handler to prevent warnings about unregistered message types from older front-end versions
+		RegisterMessageHandler("RequestQualityControl", [this](FString SourceId, FMemoryReader Ar) { /* Do nothing */ });
+
+		// Populate map
+		// Button indices found in: [REDACTED_PROFILE_URL]/webxr-input-profiles/tree/master/packages/registry/profiles
+		// HTC Vive - Left Hand
+		// Buttons
+		XRInputToFKey.Add(MakeTuple(EPixelStreaming2XRSystem::HTCVive, EControllerHand::Left, 0, EPixelStreaming2InputAction::Click), EKeys::Vive_Left_Trigger_Click);
+		XRInputToFKey.Add(MakeTuple(EPixelStreaming2XRSystem::HTCVive, EControllerHand::Left, 0, EPixelStreaming2InputAction::Axis), EKeys::Vive_Left_Trigger_Axis);
+		XRInputToFKey.Add(MakeTuple(EPixelStreaming2XRSystem::HTCVive, EControllerHand::Left, 1, EPixelStreaming2InputAction::Click), EKeys::Vive_Left_Grip_Click);
+		XRInputToFKey.Add(MakeTuple(EPixelStreaming2XRSystem::HTCVive, EControllerHand::Left, 2, EPixelStreaming2InputAction::Click), EKeys::Vive_Left_Trackpad_Click);
+		XRInputToFKey.Add(MakeTuple(EPixelStreaming2XRSystem::HTCVive, EControllerHand::Left, 2, EPixelStreaming2InputAction::Touch), EKeys::Vive_Left_Trackpad_Touch);
+		// Axes
+		XRInputToFKey.Add(MakeTuple(EPixelStreaming2XRSystem::HTCVive, EControllerHand::Left, 0, EPixelStreaming2InputAction::X), EKeys::Vive_Left_Trackpad_X);
+		XRInputToFKey.Add(MakeTuple(EPixelStreaming2XRSystem::HTCVive, EControllerHand::Left, 1, EPixelStreaming2InputAction::Y), EKeys::Vive_Left_Trackpad_Y);
+		// HTC Vive - Right Hand
+		// Buttons
+		XRInputToFKey.Add(MakeTuple(EPixelStreaming2XRSystem::HTCVive, EControllerHand::Right, 0, EPixelStreaming2InputAction::Click), EKeys::Vive_Right_Trigger_Click);
+		XRInputToFKey.Add(MakeTuple(EPixelStreaming2XRSystem::HTCVive, EControllerHand::Right, 0, EPixelStreaming2InputAction::Axis), EKeys::Vive_Right_Trigger_Axis);
+		XRInputToFKey.Add(MakeTuple(EPixelStreaming2XRSystem::HTCVive, EControllerHand::Right, 1, EPixelStreaming2InputAction::Click), EKeys::Vive_Right_Grip_Click);
+		XRInputToFKey.Add(MakeTuple(EPixelStreaming2XRSystem::HTCVive, EControllerHand::Right, 2, EPixelStreaming2InputAction::Click), EKeys::Vive_Right_Trackpad_Click);
+		XRInputToFKey.Add(MakeTuple(EPixelStreaming2XRSystem::HTCVive, EControllerHand::Right, 2, EPixelStreaming2InputAction::Touch), EKeys::Vive_Right_Trackpad_Touch);
+		// Axes
+		XRInputToFKey.Add(MakeTuple(EPixelStreaming2XRSystem::HTCVive, EControllerHand::Right, 0, EPixelStreaming2InputAction::X), EKeys::Vive_Right_Trackpad_X);
+		XRInputToFKey.Add(MakeTuple(EPixelStreaming2XRSystem::HTCVive, EControllerHand::Right, 1, EPixelStreaming2InputAction::Y), EKeys::Vive_Right_Trackpad_Y);
+
+		// Quest - Left Hand
+		// Buttons
+		XRInputToFKey.Add(MakeTuple(EPixelStreaming2XRSystem::Quest, EControllerHand::Left, 0, EPixelStreaming2InputAction::Click), EKeys::OculusTouch_Left_Trigger_Click);
+		XRInputToFKey.Add(MakeTuple(EPixelStreaming2XRSystem::Quest, EControllerHand::Left, 0, EPixelStreaming2InputAction::Axis), EKeys::OculusTouch_Left_Trigger_Axis);
+		XRInputToFKey.Add(MakeTuple(EPixelStreaming2XRSystem::Quest, EControllerHand::Left, 0, EPixelStreaming2InputAction::Touch), EKeys::OculusTouch_Left_Trigger_Touch);
+		XRInputToFKey.Add(MakeTuple(EPixelStreaming2XRSystem::Quest, EControllerHand::Left, 1, EPixelStreaming2InputAction::Click), EKeys::OculusTouch_Left_Grip_Click);
+		XRInputToFKey.Add(MakeTuple(EPixelStreaming2XRSystem::Quest, EControllerHand::Left, 1, EPixelStreaming2InputAction::Axis), EKeys::OculusTouch_Left_Grip_Axis);
+		// Index 1 (grip) touch not supported in UE
+		// Index 2 not supported by WebXR
+		XRInputToFKey.Add(MakeTuple(EPixelStreaming2XRSystem::Quest, EControllerHand::Left, 3, EPixelStreaming2InputAction::Click), EKeys::OculusTouch_Left_Thumbstick_Click);
+		XRInputToFKey.Add(MakeTuple(EPixelStreaming2XRSystem::Quest, EControllerHand::Left, 3, EPixelStreaming2InputAction::Touch), EKeys::OculusTouch_Left_Thumbstick_Touch);
+		XRInputToFKey.Add(MakeTuple(EPixelStreaming2XRSystem::Quest, EControllerHand::Left, 4, EPixelStreaming2InputAction::Click), EKeys::OculusTouch_Left_X_Click);
+		XRInputToFKey.Add(MakeTuple(EPixelStreaming2XRSystem::Quest, EControllerHand::Left, 4, EPixelStreaming2InputAction::Touch), EKeys::OculusTouch_Left_X_Touch);
+		XRInputToFKey.Add(MakeTuple(EPixelStreaming2XRSystem::Quest, EControllerHand::Left, 5, EPixelStreaming2InputAction::Click), EKeys::OculusTouch_Left_Y_Click);
+		XRInputToFKey.Add(MakeTuple(EPixelStreaming2XRSystem::Quest, EControllerHand::Left, 5, EPixelStreaming2InputAction::Touch), EKeys::OculusTouch_Left_Y_Touch);
+		// Index 6 (thumbrest) not supported in UE
+
+		// Axes
+		// Indices 0 and 1 not supported in WebXR
+		XRInputToFKey.Add(MakeTuple(EPixelStreaming2XRSystem::Quest, EControllerHand::Left, 2, EPixelStreaming2InputAction::X), EKeys::OculusTouch_Left_Thumbstick_X);
+		XRInputToFKey.Add(MakeTuple(EPixelStreaming2XRSystem::Quest, EControllerHand::Left, 3, EPixelStreaming2InputAction::Y), EKeys::OculusTouch_Left_Thumbstick_Y);
+
+		// Quest - Right Hand
+		// Buttons
+		XRInputToFKey.Add(MakeTuple(EPixelStreaming2XRSystem::Quest, EControllerHand::Right, 0, EPixelStreaming2InputAction::Click), EKeys::OculusTouch_Right_Trigger_Click);
+		XRInputToFKey.Add(MakeTuple(EPixelStreaming2XRSystem::Quest, EControllerHand::Right, 0, EPixelStreaming2InputAction::Axis), EKeys::OculusTouch_Right_Trigger_Axis);
+		XRInputToFKey.Add(MakeTuple(EPixelStreaming2XRSystem::Quest, EControllerHand::Right, 0, EPixelStreaming2InputAction::Touch), EKeys::OculusTouch_Right_Trigger_Touch);
+		XRInputToFKey.Add(MakeTuple(EPixelStreaming2XRSystem::Quest, EControllerHand::Right, 1, EPixelStreaming2InputAction::Click), EKeys::OculusTouch_Right_Grip_Click);
+		XRInputToFKey.Add(MakeTuple(EPixelStreaming2XRSystem::Quest, EControllerHand::Right, 1, EPixelStreaming2InputAction::Axis), EKeys::OculusTouch_Right_Grip_Axis);
+		// Index 1 (grip) touch not supported in UE
+		// Index 2 not supported by WebXR
+		XRInputToFKey.Add(MakeTuple(EPixelStreaming2XRSystem::Quest, EControllerHand::Right, 3, EPixelStreaming2InputAction::Click), EKeys::OculusTouch_Right_Thumbstick_Click);
+		XRInputToFKey.Add(MakeTuple(EPixelStreaming2XRSystem::Quest, EControllerHand::Right, 3, EPixelStreaming2InputAction::Touch), EKeys::OculusTouch_Right_Thumbstick_Touch);
+		XRInputToFKey.Add(MakeTuple(EPixelStreaming2XRSystem::Quest, EControllerHand::Right, 4, EPixelStreaming2InputAction::Click), EKeys::OculusTouch_Right_A_Click);
+		XRInputToFKey.Add(MakeTuple(EPixelStreaming2XRSystem::Quest, EControllerHand::Right, 4, EPixelStreaming2InputAction::Touch), EKeys::OculusTouch_Right_A_Touch);
+		XRInputToFKey.Add(MakeTuple(EPixelStreaming2XRSystem::Quest, EControllerHand::Right, 5, EPixelStreaming2InputAction::Click), EKeys::OculusTouch_Right_B_Click);
+		XRInputToFKey.Add(MakeTuple(EPixelStreaming2XRSystem::Quest, EControllerHand::Right, 5, EPixelStreaming2InputAction::Touch), EKeys::OculusTouch_Right_B_Touch);
+		// Index 6 (thumbrest) not supported in UE
+
+		// Axes
+		// Indices 0 and 1 not supported in WebXR
+		XRInputToFKey.Add(MakeTuple(EPixelStreaming2XRSystem::Quest, EControllerHand::Right, 2, EPixelStreaming2InputAction::X), EKeys::OculusTouch_Right_Thumbstick_X);
+		XRInputToFKey.Add(MakeTuple(EPixelStreaming2XRSystem::Quest, EControllerHand::Right, 3, EPixelStreaming2InputAction::Y), EKeys::OculusTouch_Right_Thumbstick_Y);
+
+		// Gamepad Axes
+		GamepadInputToFKey.Add(MakeTuple(1, EPixelStreaming2InputAction::Axis), EKeys::Gamepad_LeftX);
+		GamepadInputToFKey.Add(MakeTuple(2, EPixelStreaming2InputAction::Axis), EKeys::Gamepad_LeftY);
+		GamepadInputToFKey.Add(MakeTuple(3, EPixelStreaming2InputAction::Axis), EKeys::Gamepad_RightX);
+		GamepadInputToFKey.Add(MakeTuple(4, EPixelStreaming2InputAction::Axis), EKeys::Gamepad_RightY);
+		GamepadInputToFKey.Add(MakeTuple(5, EPixelStreaming2InputAction::Axis), EKeys::Gamepad_LeftTriggerAxis);
+		GamepadInputToFKey.Add(MakeTuple(6, EPixelStreaming2InputAction::Axis), EKeys::Gamepad_RightTriggerAxis);
+		// Gamepad Buttons
+		GamepadInputToFKey.Add(MakeTuple(0, EPixelStreaming2InputAction::Click), EKeys::Gamepad_FaceButton_Bottom);
+		GamepadInputToFKey.Add(MakeTuple(1, EPixelStreaming2InputAction::Click), EKeys::Gamepad_FaceButton_Right);
+		GamepadInputToFKey.Add(MakeTuple(2, EPixelStreaming2InputAction::Click), EKeys::Gamepad_FaceButton_Left);
+		GamepadInputToFKey.Add(MakeTuple(3, EPixelStreaming2InputAction::Click), EKeys::Gamepad_FaceButton_Top);
+		GamepadInputToFKey.Add(MakeTuple(4, EPixelStreaming2InputAction::Click), EKeys::Gamepad_LeftShoulder);
+		GamepadInputToFKey.Add(MakeTuple(5, EPixelStreaming2InputAction::Click), EKeys::Gamepad_RightShoulder);
+		GamepadInputToFKey.Add(MakeTuple(6, EPixelStreaming2InputAction::Click), EKeys::Gamepad_LeftTrigger);
+		GamepadInputToFKey.Add(MakeTuple(7, EPixelStreaming2InputAction::Click), EKeys::Gamepad_RightTrigger);
+		GamepadInputToFKey.Add(MakeTuple(8, EPixelStreaming2InputAction::Click), EKeys::Gamepad_Special_Left);
+		GamepadInputToFKey.Add(MakeTuple(9, EPixelStreaming2InputAction::Click), EKeys::Gamepad_Special_Right);
+		GamepadInputToFKey.Add(MakeTuple(10, EPixelStreaming2InputAction::Click), EKeys::Gamepad_LeftThumbstick);
+		GamepadInputToFKey.Add(MakeTuple(11, EPixelStreaming2InputAction::Click), EKeys::Gamepad_RightThumbstick);
+		GamepadInputToFKey.Add(MakeTuple(12, EPixelStreaming2InputAction::Click), EKeys::Gamepad_DPad_Up);
+		GamepadInputToFKey.Add(MakeTuple(13, EPixelStreaming2InputAction::Click), EKeys::Gamepad_DPad_Down);
+		GamepadInputToFKey.Add(MakeTuple(14, EPixelStreaming2InputAction::Click), EKeys::Gamepad_DPad_Left);
+		GamepadInputToFKey.Add(MakeTuple(15, EPixelStreaming2InputAction::Click), EKeys::Gamepad_DPad_Right);
+
+		PopulateDefaultCommandHandlers();
+	}
+
+	FRTCInputHandler::~FRTCInputHandler()
+	{
+	}
+
+	void FRTCInputHandler::HandleOnKeyChar(FMemoryReader Ar)
+	{
+		TPayload<TCHAR> Payload(Ar);
+		OnKeyChar(Payload.Get<0>());
+	}
+
+	void FRTCInputHandler::HandleOnKeyDown(FMemoryReader Ar)
+	{
+		TPayload<uint8, uint8> Payload(Ar);
+
+		bool		bIsRepeat = Payload.Get<1>() != 0;
+		const FKey* AgnosticKey = UE::PixelStreaming2Input::JavaScriptKeyCodeToFKey[Payload.Get<0>()];
+		if (FilterKey(*AgnosticKey))
+		{
+			OnKeyDown(*AgnosticKey, bIsRepeat);
+		}
+	}
+
+	void FRTCInputHandler::HandleOnKeyUp(FMemoryReader Ar)
+	{
+		TPayload<uint8> Payload(Ar);
+		const FKey*		AgnosticKey = UE::PixelStreaming2Input::JavaScriptKeyCodeToFKey[Payload.Get<0>()];
+		if (FilterKey(*AgnosticKey))
+		{
+			OnKeyUp(*AgnosticKey);
+		}
+	}
+
+	void FRTCInputHandler::HandleOnMouseEnter(FMemoryReader Ar)
+	{
+		OnMouseEnter();
+	}
+
+	void FRTCInputHandler::HandleOnMouseLeave(FMemoryReader Ar)
+	{
+		OnMouseLeave();
+	}
+
+	void FRTCInputHandler::HandleOnMouseDown(FMemoryReader Ar)
+	{
+		TPayload<uint8, uint16, uint16> Payload(Ar);
+
+		EMouseButtons::Type Button = static_cast<EMouseButtons::Type>(Payload.Get<0>());
+		//                                                                           convert range from 0,65536 -> 0,1
+		FIntPoint ScreenLocation = ConvertFromNormalizedScreenLocation(FVector2D(Payload.Get<1>() / uint16_MAX, Payload.Get<2>() / uint16_MAX));
+
+		OnMouseDown(Button, ScreenLocation);
+	}
+
+	void FRTCInputHandler::HandleOnMouseUp(FMemoryReader Ar)
+	{
+		TPayload<uint8, uint16, uint16> Payload(Ar);
+
+		EMouseButtons::Type Button = static_cast<EMouseButtons::Type>(Payload.Get<0>());
+		OnMouseUp(Button);
+	}
+
+	void FRTCInputHandler::HandleOnMouseMove(FMemoryReader Ar)
+	{
+		TPayload<uint16, uint16, int16, int16> Payload(Ar);
+
+		//                                                                           convert range from 0,65536 -> 0,1
+		FIntPoint ScreenLocation = ConvertFromNormalizedScreenLocation(FVector2D(Payload.Get<0>() / uint16_MAX, Payload.Get<1>() / uint16_MAX));
+		//                                                                 convert range from -32,768 to 32,767 -> -1,1
+		FIntPoint Delta = ConvertFromNormalizedScreenLocation(FVector2D(Payload.Get<2>() / int16_MAX, Payload.Get<3>() / int16_MAX), false);
+
+		OnMouseMove(ScreenLocation, Delta);
+	}
+
+	void FRTCInputHandler::HandleOnMouseWheel(FMemoryReader Ar)
+	{
+		TPayload<int16, uint16, uint16> Payload(Ar);
+		//                                                                           convert range from 0,65536 -> 0,1
+		FIntPoint	ScreenLocation = ConvertFromNormalizedScreenLocation(FVector2D(Payload.Get<1>() / uint16_MAX, Payload.Get<2>() / uint16_MAX));
+		const float SpinFactor = 1 / 120.0f;
+
+		OnMouseWheel(ScreenLocation, Payload.Get<0>() * SpinFactor);
+	}
+
+	void FRTCInputHandler::HandleOnMouseDoubleClick(FMemoryReader Ar)
+	{
+		TPayload<uint8, uint16, uint16> Payload(Ar);
+		EMouseButtons::Type				Button = static_cast<EMouseButtons::Type>(Payload.Get<0>());
+
+		//                                                                           convert range from 0,65536 -> 0,1
+		FIntPoint ScreenLocation = ConvertFromNormalizedScreenLocation(FVector2D(Payload.Get<1>() / uint16_MAX, Payload.Get<2>() / uint16_MAX));
+
+		OnMouseDoubleClick(Button, ScreenLocation);
+	}
+
+	void FRTCInputHandler::HandleOnTouchStarted(FMemoryReader Ar)
+	{
+		TPayload<uint8> Payload(Ar);
+
+		uint8 NumTouches = Payload.Get<0>();
+		for (uint8 TouchIdx = 0; TouchIdx < NumTouches; TouchIdx++)
+		{
+			//       PosX    PoxY    IDX   Force  Valid
+			TPayload<uint16, uint16, uint8, uint8, uint8> Touch(Ar);
+			// If Touch is valid
+			if (Touch.Get<4>() != 0)
+			{
+				//                                                                           convert range from 0,65536 -> 0,1
+				FIntPoint	TouchLocation = ConvertFromNormalizedScreenLocation(FVector2D(Touch.Get<0>() / uint16_MAX, Touch.Get<1>() / uint16_MAX));
+				const int32 TouchIndex = Touch.Get<2>();
+				const float TouchForce = Touch.Get<3>() / 255.0f;
+
+				OnTouchStarted(TouchLocation, TouchIndex, TouchForce);
+			}
+		}
+	}
+
+	void FRTCInputHandler::HandleOnTouchMoved(FMemoryReader Ar)
+	{
+		TPayload<uint8> Payload(Ar);
+
+		uint8 NumTouches = Payload.Get<0>();
+		for (uint8 TouchIdx = 0; TouchIdx < NumTouches; TouchIdx++)
+		{
+			//       PosX    PoxY    IDX   Force  Valid
+			TPayload<uint16, uint16, uint8, uint8, uint8> Touch(Ar);
+			// If Touch is valid
+			if (Touch.Get<4>() != 0)
+			{
+				//                                                                           convert range from 0,65536 -> 0,1
+				FIntPoint	TouchLocation = ConvertFromNormalizedScreenLocation(FVector2D(Touch.Get<0>() / uint16_MAX, Touch.Get<1>() / uint16_MAX));
+				const int32 TouchIndex = Touch.Get<2>();
+				const float TouchForce = Touch.Get<3>() / 255.0f;
+
+				OnTouchMoved(TouchLocation, TouchIndex, TouchForce);
+			}
+		}
+	}
+
+	void FRTCInputHandler::HandleOnTouchEnded(FMemoryReader Ar)
+	{
+		TPayload<uint8> Payload(Ar);
+
+		uint8 NumTouches = Payload.Get<0>();
+		for (uint8 TouchIdx = 0; TouchIdx < NumTouches; TouchIdx++)
+		{
+			//       PosX    PoxY    IDX   Force  Valid
+			TPayload<uint16, uint16, uint8, uint8, uint8> Touch(Ar);
+			// Always allowing the "up" events regardless of in or outside the valid region so
+			// states aren't stuck "down". Might want to uncomment this if it causes other issues.
+			// if (Touch.Get<4>() != 0)
+			{
+				//                                                                           convert range from 0,65536 -> 0,1
+				FIntPoint	TouchLocation = ConvertFromNormalizedScreenLocation(FVector2D(Touch.Get<0>() / uint16_MAX, Touch.Get<1>() / uint16_MAX));
+				const int32 TouchIndex = Touch.Get<2>();
+
+				OnTouchEnded(TouchLocation, TouchIndex);
+			}
+		}
+	}
+
+	void FRTCInputHandler::HandleOnControllerConnected(FMemoryReader Ar)
+	{
+		OnControllerConnected();
+	}
+
+	void FRTCInputHandler::HandleOnControllerAnalog(FMemoryReader Ar)
+	{
+		TPayload<uint8, uint8, double> Payload(Ar);
+
+		uint8  ControllerIndex = Payload.Get<0>();
+		uint8  KeyId = Payload.Get<1>();
+		double AxisValue = Payload.Get<2>();
+
+		FKey* KeyPtr = GamepadInputToFKey.Find(MakeTuple(KeyId, EPixelStreaming2InputAction::Axis));
+		if (!KeyPtr)
+		{
+			return;
+		}
+
+		OnControllerAnalog(ControllerIndex, *KeyPtr, AxisValue);
+	}
+
+	void FRTCInputHandler::HandleOnControllerButtonPressed(FMemoryReader Ar)
+	{
+		TPayload<uint8, uint8, uint8> Payload(Ar);
+
+		uint8 ControllerIndex = Payload.Get<0>();
+		uint8 KeyId = Payload.Get<1>();
+		bool  bIsRepeat = Payload.Get<2>() != 0;
+
+		FKey* KeyPtr = GamepadInputToFKey.Find(MakeTuple(KeyId, EPixelStreaming2InputAction::Click));
+		if (KeyPtr == nullptr)
+		{
+			return;
+		}
+
+		OnControllerButtonPressed(ControllerIndex, *KeyPtr, bIsRepeat);
+	}
+
+	void FRTCInputHandler::HandleOnControllerButtonReleased(FMemoryReader Ar)
+	{
+		TPayload<uint8, uint8> Payload(Ar);
+
+		uint8 ControllerIndex = Payload.Get<0>();
+		uint8 KeyId = Payload.Get<1>();
+
+		FKey* KeyPtr = GamepadInputToFKey.Find(MakeTuple(KeyId, EPixelStreaming2InputAction::Click));
+		if (KeyPtr == nullptr)
+		{
+			return;
+		}
+
+		OnControllerButtonReleased(ControllerIndex, *KeyPtr);
+	}
+
+	void FRTCInputHandler::HandleOnControllerDisconnected(FMemoryReader Ar)
+	{
+		TPayload<uint8> Payload(Ar);
+
+		OnControllerDisconnected(Payload.Get<0>());
+	}
+
+	void FRTCInputHandler::HandleOnXREyeViews(FMemoryReader Ar)
+	{
+		// The `Ar` buffer contains the left eye transform matrix stored as 16 floats
+		FTransform LeftEyeTransform = WebXRMatrixToUETransform(ExtractWebXRMatrix(Ar));
+		// The `Ar` buffer contains the left eye projection matrix stored as 16 floats
+		FMatrix LeftEyeProjectionMatrix = ExtractWebXRMatrix(Ar);
+		// The `Ar` buffer contains the right eye transform matrix stored as 16 floats
+		FTransform RightEyeTransform = WebXRMatrixToUETransform(ExtractWebXRMatrix(Ar));
+		// The `Ar` buffer contains the right eye projection matrix stored as 16 floats
+		FMatrix RightEyeProjectionMatrix = ExtractWebXRMatrix(Ar);
+
+		// The `Ar` buffer contains the right eye projection matrix stored as 16 floats
+		FTransform HMDTransform = WebXRMatrixToUETransform(ExtractWebXRMatrix(Ar));
+
+		OnXREyeViews(LeftEyeTransform, LeftEyeProjectionMatrix, RightEyeTransform, RightEyeProjectionMatrix, HMDTransform);
+	}
+
+	void FRTCInputHandler::HandleOnXRHMDTransform(FMemoryReader Ar)
+	{
+		// The `Ar` buffer contains the transform matrix stored as 16 floats
+		FTransform HMDTransform = WebXRMatrixToUETransform(ExtractWebXRMatrix(Ar));
+
+		OnXRHMDTransform(HMDTransform);
+	}
+
+	void FRTCInputHandler::HandleOnXRControllerTransform(FMemoryReader Ar)
+	{
+		// The `Ar` buffer contains the transform matrix stored as 16 floats
+		FTransform ControllerTransform = WebXRMatrixToUETransform(ExtractWebXRMatrix(Ar));
+
+		// The `Ar` buffer contains a UInt8 for the handedness
+		EControllerHand Handedness = EControllerHand::Left;
+		Ar << Handedness;
+
+		OnXRControllerTransform(ControllerTransform, Handedness);
+	}
+
+	void FRTCInputHandler::HandleOnXRButtonTouched(FMemoryReader Ar)
+	{
+		TPayload<uint8, uint8, uint8> Payload(Ar);
+
+		EControllerHand Handedness = static_cast<EControllerHand>(Payload.Get<0>());
+		uint8			ButtonIdx = Payload.Get<1>();
+		bool			bIsRepeat = Payload.Get<2>() != 0;
+
+		EPixelStreaming2XRSystem System = IPixelStreaming2HMDModule::Get().GetActiveXRSystem();
+
+		FKey* KeyPtr = XRInputToFKey.Find(MakeTuple(System, Handedness, ButtonIdx, EPixelStreaming2InputAction::Touch));
+		if (KeyPtr == nullptr)
+		{
+			return;
+		}
+
+		OnXRButtonTouched(Handedness, *KeyPtr, bIsRepeat);
+	}
+
+	void FRTCInputHandler::HandleOnXRButtonTouchReleased(FMemoryReader Ar)
+	{
+		TPayload<uint8, uint8, uint8> Payload(Ar);
+
+		EControllerHand Handedness = static_cast<EControllerHand>(Payload.Get<0>());
+		uint8			ButtonIdx = Payload.Get<1>();
+
+		EPixelStreaming2XRSystem System = IPixelStreaming2HMDModule::Get().GetActiveXRSystem();
+
+		FKey* KeyPtr = XRInputToFKey.Find(MakeTuple(System, Handedness, ButtonIdx, EPixelStreaming2InputAction::Touch));
+		if (KeyPtr == nullptr)
+		{
+			return;
+		}
+
+		OnXRButtonTouchReleased(Handedness, *KeyPtr);
+	}
+
+	void FRTCInputHandler::HandleOnXRButtonPressed(FMemoryReader Ar)
+	{
+		TPayload<uint8, uint8, uint8, double> Payload(Ar);
+
+		EControllerHand Handedness = static_cast<EControllerHand>(Payload.Get<0>());
+		uint8			ButtonIdx = Payload.Get<1>();
+		bool			bIsRepeat = Payload.Get<2>() != 0;
+		double			AnalogValue = Payload.Get<3>();
+
+		EPixelStreaming2XRSystem System = IPixelStreaming2HMDModule::Get().GetActiveXRSystem();
+
+		FKey* ButtonPtr = XRInputToFKey.Find(MakeTuple(System, Handedness, ButtonIdx, EPixelStreaming2InputAction::Click));
+		if (ButtonPtr)
+		{
+			OnXRButtonPressed(Handedness, *ButtonPtr, bIsRepeat);
+		}
+
+		// Try and see if there is an axis associated with this button (usually the case for triggers)
+		FKey* AxisPtr = XRInputToFKey.Find(MakeTuple(System, Handedness, ButtonIdx, EPixelStreaming2InputAction::Axis));
+		// If we have axis associate with this press then set axis value to the button press value
+		if (AxisPtr)
+		{
+			OnXRAnalog(Handedness, *AxisPtr, AnalogValue);
+		}
+	}
+
+	void FRTCInputHandler::HandleOnXRButtonReleased(FMemoryReader Ar)
+	{
+		TPayload<uint8, uint8, uint8> Payload(Ar);
+
+		EControllerHand Handedness = static_cast<EControllerHand>(Payload.Get<0>());
+		uint8			ButtonIdx = Payload.Get<1>();
+
+		EPixelStreaming2XRSystem System = IPixelStreaming2HMDModule::Get().GetActiveXRSystem();
+
+		// Try and see if there is an axis associated with this button (usually the case for triggers)
+		FKey* AxisPtr = XRInputToFKey.Find(MakeTuple(System, Handedness, ButtonIdx, EPixelStreaming2InputAction::Axis));
+		// If we have axis associate with this press then set axis value to 0.0
+		if (AxisPtr)
+		{
+			OnXRAnalog(Handedness, *AxisPtr, 0.0f);
+		}
+
+		// Do the actual release after the analog trigger, as the release can cancel any further inputs
+		FKey* ButtonPtr = XRInputToFKey.Find(MakeTuple(System, Handedness, ButtonIdx, EPixelStreaming2InputAction::Click));
+		if (ButtonPtr)
+		{
+			OnXRButtonReleased(Handedness, *ButtonPtr);
+		}
+	}
+
+	void FRTCInputHandler::HandleOnXRAnalog(FMemoryReader Ar)
+	{
+		TPayload<uint8, uint8, double> Payload(Ar);
+
+		EControllerHand Handedness = static_cast<EControllerHand>(Payload.Get<0>());
+		uint8			AxisIndex = Payload.Get<1>();
+		double			AnalogValue = Payload.Get<2>();
+
+		EPixelStreaming2XRSystem	System = IPixelStreaming2HMDModule::Get().GetActiveXRSystem();
+		EPixelStreaming2InputAction InputAction = static_cast<EPixelStreaming2InputAction>(AxisIndex % 2);
+
+		FKey* KeyPtr = XRInputToFKey.Find(MakeTuple(System, Handedness, AxisIndex, InputAction));
+		if (KeyPtr == nullptr)
+		{
+			return;
+		}
+
+		OnXRAnalog(Handedness, *KeyPtr, AnalogValue);
+	}
+
+	void FRTCInputHandler::HandleOnXRSystem(FMemoryReader Ar)
+	{
+		TPayload<uint8> Payload(Ar);
+
+		EPixelStreaming2XRSystem System = static_cast<EPixelStreaming2XRSystem>(Payload.Get<0>());
+
+		OnXRSystem(System);
+	}
+
+	/**
+	 * Command handling
+	 */
+	void FRTCInputHandler::HandleOnCommand(FString SourceId, FMemoryReader Ar)
+	{
+		FString Res;
+		Res.GetCharArray().SetNumUninitialized(Ar.TotalSize() / 2 + 1);
+		Ar.Serialize(Res.GetCharArray().GetData(), Ar.TotalSize());
+
+		FString Descriptor = Res.Mid(MessageHeaderOffset);
+		UE_LOGFMT(LogPixelStreaming2RTC, Verbose, "Command: {0}", Descriptor);
+
+		// Iterate each command handler and see if the command we got matches any of the bound command names.
+		for (auto& CommandHandlersPair : CommandHandlers)
+		{
+			FString CommandValue;
+			bool	bSuccess = false;
+			ExtractJsonFromDescriptor(Descriptor, CommandHandlersPair.Key, CommandValue, bSuccess);
+			if (bSuccess)
+			{
+				// Execute bound command handler with descriptor and parsed command value
+				CommandHandlersPair.Value(SourceId, Descriptor, CommandValue);
+				return;
+			}
+		}
+	}
+
+	/**
+	 * UI Interaction handling
+	 */
+	void FRTCInputHandler::HandleUIInteraction(FMemoryReader Ar)
+	{
+		// FPixelStreaming2Module overwrites this handler
+	}
+
+	/**
+	 * Textbox Entry handling
+	 */
+	void FRTCInputHandler::HandleOnTextboxEntry(FMemoryReader Ar)
+	{
+		FString Res;
+		Res.GetCharArray().SetNumUninitialized(Ar.TotalSize() / 2 + 1);
+		Ar.Serialize(Res.GetCharArray().GetData(), Ar.TotalSize());
+		FString Text = Res.Mid(1);
+
+		FSlateApplication::Get().ForEachUser([this, Text](FSlateUser& User) {
+			TSharedPtr<SWidget> FocusedWidget = User.GetFocusedWidget();
+
+			bool bIsEditableTextType = FocusedWidget->GetType() == TEXT("SEditableText");
+			bool bIsMultiLineEditableTextType = FocusedWidget->GetType() == TEXT("SMultiLineEditableText");
+			bool bEditable = FocusedWidget && (bIsEditableTextType || bIsMultiLineEditableTextType);
+
+			if (bEditable)
+			{
+				if (bIsEditableTextType)
+				{
+					SEditableText* TextBox = static_cast<SEditableText*>(FocusedWidget.Get());
+					TextBox->SetText(FText::FromString(Text));
+				}
+				else if (bIsMultiLineEditableTextType)
+				{
+					SMultiLineEditableTextBox* TextBox = static_cast<SMultiLineEditableTextBox*>(FocusedWidget.Get());
+					TextBox->SetText(FText::FromString(Text));
+				}
+
+				// We need to manually trigger an Enter key press so that the OnTextCommitted delegate gets fired
+				const uint32* KeyPtr = nullptr;
+				const uint32* CharacterPtr = nullptr;
+				FInputKeyManager::Get().GetCodesFromKey(EKeys::Enter, KeyPtr, CharacterPtr);
+				uint32 Key = KeyPtr ? *KeyPtr : 0;
+				uint32 Character = CharacterPtr ? *CharacterPtr : 0;
+				if (Key != 0 || Character != 0)
+				{
+					MessageHandler->OnKeyDown((int32)Key, (int32)Character, false);
+					MessageHandler->OnKeyUp((int32)Key, (int32)Character, false);
+				}
+			}
+		});
+	}
+
+	void FRTCInputHandler::PopulateDefaultCommandHandlers()
+	{
+		// Execute console commands if passed "ConsoleCommand" and -PixelStreaming2AllowConsoleCommands is on.
+		CommandHandlers.Add(TEXT("ConsoleCommand"), [this](FString SourceId, FString Descriptor, FString ConsoleCommand) {
+			if (!UPixelStreaming2PluginSettings::CVarInputAllowConsoleCommands.GetValueOnAnyThread()
+				|| !IsElevated(SourceId))
+			{
+				return;
+			}
+			GEngine->Exec(GEngine->GetWorld(), *ConsoleCommand);
+		});
+
+		// Change width/height if sent { "Resolution.Width": 1920, "Resolution.Height": 1080 }
+		CommandHandlers.Add(TEXT("Resolution.Width"), [this](FString SourceId, FString Descriptor, FString WidthString) {
+			bool	bSuccess = false;
+			FString HeightString;
+			ExtractJsonFromDescriptor(Descriptor, TEXT("Resolution.Height"), HeightString, bSuccess);
+			if (bSuccess && IsElevated(SourceId))
+			{
+				int Width = FCString::Atoi(*WidthString);
+				int Height = FCString::Atoi(*HeightString);
+				if (Width < 1 || Height < 1)
+				{
+					return;
+				}
+
+				FString ChangeResCommand = FString::Printf(TEXT("r.SetRes %dx%d"), Width, Height);
+				GEngine->Exec(GEngine->GetWorld(), *ChangeResCommand);
+			}
+		});
+
+		// Response to "Stat.FPS" by calling "stat fps"
+		CommandHandlers.Add(TEXT("Stat.FPS"), [](FString SourceId, FString Descriptor, FString FPSCommand) {
+			FString StatFPSCommand = FString::Printf(TEXT("stat fps"));
+			GEngine->Exec(GEngine->GetWorld(), *StatFPSCommand);
+		});
+	}
+
+	FMatrix FRTCInputHandler::ExtractWebXRMatrix(FMemoryReader& Ar)
+	{
+		FMatrix OutMat;
+		for (int32 Row = 0; Row < 4; ++Row)
+
+		{
+			float Col0 = 0.0f, Col1 = 0.0f, Col2 = 0.0f, Col3 = 0.0f;
+			Ar << Col0 << Col1 << Col2 << Col3;
+			OutMat.M[Row][0] = Col0;
+			OutMat.M[Row][1] = Col1;
+			OutMat.M[Row][2] = Col2;
+			OutMat.M[Row][3] = Col3;
+		}
+		OutMat.DiagnosticCheckNaN();
+		return OutMat;
+	}
+
+	FTransform FRTCInputHandler::WebXRMatrixToUETransform(FMatrix Mat)
+	{
+		// Rows and columns are swapped between raw mat and FMat
+		FMatrix UEMatrix = FMatrix(
+			FPlane(Mat.M[0][0], Mat.M[1][0], Mat.M[2][0], Mat.M[3][0]),
+			FPlane(Mat.M[0][1], Mat.M[1][1], Mat.M[2][1], Mat.M[3][1]),
+			FPlane(Mat.M[0][2], Mat.M[1][2], Mat.M[2][2], Mat.M[3][2]),
+			FPlane(Mat.M[0][3], Mat.M[1][3], Mat.M[2][3], Mat.M[3][3]));
+
+		// Extract scale vector and reorder coordinates to be UE coordinate system.
+		FVector ScaleVectorRaw = UEMatrix.GetScaleVector();
+		// Note: We do not invert Z scaling here because we already handle that when we rebuild translation/rot below.
+		FVector ScaleVector = FVector(ScaleVectorRaw.Z, ScaleVectorRaw.X, ScaleVectorRaw.Y);
+
+		// Temporarily remove scaling component as we need rotation axes to be unit length for proper quat conversion
+		UEMatrix.RemoveScaling();
+
+		// Extract & convert translation component to UE coordinate syste,
+		FVector Translation = FVector(-UEMatrix.M[3][2], UEMatrix.M[3][0], UEMatrix.M[3][1]) * 100.0f;
+
+		// Extract & convert rotation component to UE coordinate system
+		FQuat RawRotation(UEMatrix);
+		FQuat Rotation(-RawRotation.Z, RawRotation.X, RawRotation.Y, -RawRotation.W);
+		return FTransform(Rotation, Translation, ScaleVector);
+	}
+} // namespace UE::PixelStreaming2
\ No newline at end of file
diff --git a/Engine/Plugins/Media/PixelStreaming2/Source/PixelStreaming2RTC/Private/RTCInputHandler.h b/Engine/Plugins/Media/PixelStreaming2/Source/PixelStreaming2RTC/Private/RTCInputHandler.h
new file mode 100644
--- /dev/null
+++ b/Engine/Plugins/Media/PixelStreaming2/Source/PixelStreaming2RTC/Private/RTCInputHandler.h
@@ -0,0 +1,95 @@
+// Copyright Epic Games, Inc. All Rights Reserved.
+
+#pragma once
+
+#include "InputHandler.h"
+
+namespace UE::PixelStreaming2
+{
+	class FRTCInputHandler : public UE::PixelStreaming2Input::FPixelStreaming2InputHandler
+	{
+	public:
+		static TSharedPtr<FRTCInputHandler> Create();
+
+		virtual ~FRTCInputHandler();
+
+	protected:
+		FRTCInputHandler();
+		/**
+		 * Touch handling
+		 */
+		void HandleOnTouchStarted(FMemoryReader Ar);
+		void HandleOnTouchMoved(FMemoryReader Ar);
+		void HandleOnTouchEnded(FMemoryReader Ar);
+		/**
+		 * Key press handling
+		 */
+		void HandleOnKeyChar(FMemoryReader Ar);
+		void HandleOnKeyDown(FMemoryReader Ar);
+		void HandleOnKeyUp(FMemoryReader Ar);
+		/**
+		 * Mouse handling
+		 */
+		void HandleOnMouseEnter(FMemoryReader Ar);
+		void HandleOnMouseLeave(FMemoryReader Ar);
+		void HandleOnMouseDown(FMemoryReader Ar);
+		void HandleOnMouseUp(FMemoryReader Ar);
+		void HandleOnMouseMove(FMemoryReader Ar);
+		void HandleOnMouseWheel(FMemoryReader Ar);
+		void HandleOnMouseDoubleClick(FMemoryReader Ar);
+		/**
+		 * Controller handling
+		 */
+		void HandleOnControllerConnected(FMemoryReader Ar);
+		void HandleOnControllerAnalog(FMemoryReader Ar);
+		void HandleOnControllerButtonPressed(FMemoryReader Ar);
+		void HandleOnControllerButtonReleased(FMemoryReader Ar);
+		void HandleOnControllerDisconnected(FMemoryReader Ar);
+		/**
+		 * XR handling
+		 */
+		void HandleOnXREyeViews(FMemoryReader Ar);
+		void HandleOnXRHMDTransform(FMemoryReader Ar);
+		void HandleOnXRControllerTransform(FMemoryReader Ar);
+		void HandleOnXRButtonTouched(FMemoryReader Ar);
+		void HandleOnXRButtonTouchReleased(FMemoryReader Ar);
+		void HandleOnXRButtonPressed(FMemoryReader Ar);
+		void HandleOnXRButtonReleased(FMemoryReader Ar);
+		void HandleOnXRAnalog(FMemoryReader Ar);
+		void HandleOnXRSystem(FMemoryReader Ar);
+		/**
+		 * Command handling
+		 */
+		void HandleOnCommand(FString SourceId, FMemoryReader Ar);
+		/**
+		 * UI Interaction handling
+		 */
+		void HandleUIInteraction(FMemoryReader Ar);
+		/**
+		 * Textbox Entry handling
+		 */
+		void HandleOnTextboxEntry(FMemoryReader Ar);
+
+	private:
+		/**
+		 * Populate default command handlers for data channel messages sent with "{ type: "Command" }".
+		 */
+		void PopulateDefaultCommandHandlers();
+
+		/**
+		 * Extract 4x4 WebXR ordered matrix and convert to FMatrix.
+		 */
+		FMatrix ExtractWebXRMatrix(FMemoryReader& Ar);
+
+		/**
+		 * Converts the 'Y up' 'right handed' WebXR coordinate system transform to Unreal's 'Z up'
+		 * 'left handed' coordinate system. Note: Ignores scale.
+		 * Assumes WebXR conforms to the following: (https://developer.mozilla.org/en-US/docs/Web/API/WebXR_Device_API/Geometry)
+		 * @return A 4x4 z-up transform matrix for use with UE.
+		 */
+		FTransform WebXRMatrixToUETransform(FMatrix Mat);
+
+		TMap<TTuple<EPixelStreaming2XRSystem, EControllerHand, uint8, EPixelStreaming2InputAction>, FKey> XRInputToFKey;
+		TMap<TTuple<uint8, EPixelStreaming2InputAction>, FKey>											  GamepadInputToFKey;
+	};
+} // namespace UE::PixelStreaming2
\ No newline at end of file
