[PS2] Refactor: multiple input related components to enable differing use cases with PS2
- The InputComponents array has moved from being on the FPixelStreaming2Module to a global, internal, variable.
- UPixelStreaming2InputComponent BP has moved from ./Private/Blueprints to ./Internal/Blueprints.
- FPixelStreaming2Module no longer registers default command handlers or handling of UIInteractions. This is instead now done in the FEpicRtcStreamer constructor. This is done so different streamer implementations can register their own UIInteraction handler as well as command handlers.
- A handful of state variables like focused/unfocused pos have been added back to the base InputHandler.
- The logic to scale a point on the screen into the range of 0-1 has been refactored into a separate method (ConvertToNormalizedScreenLocation).
- A new method called SynthesizeMouseMove has been added to the base handler.

diff --git a/Engine/Plugins/Media/PixelStreaming2/Source/PixelStreaming2/Internal/Blueprints/PixelStreaming2InputComponent.h b/Engine/Plugins/Media/PixelStreaming2/Source/PixelStreaming2/Internal/Blueprints/PixelStreaming2InputComponent.h
new file mode 100644
--- /dev/null
+++ b/Engine/Plugins/Media/PixelStreaming2/Source/PixelStreaming2/Internal/Blueprints/PixelStreaming2InputComponent.h
@@ -0,0 +1,61 @@
+// Copyright Epic Games, Inc. All Rights Reserved.
+
+#pragma once
+
+#include "Components/ActorComponent.h"
+#include "PixelStreaming2InputComponent.generated.h"
+
+/**
+ * This component may be attached to an actor to allow UI interactions to be
+ * handled as the delegate will be notified about the interaction and will be
+ * supplied with a generic descriptor string containing, for example, JSON data.
+ * Responses back to the source of the UI interactions may also be sent.
+ */
+UCLASS(Blueprintable, ClassGroup = (PixelStreaming2), meta = (BlueprintSpawnableComponent))
+class UPixelStreaming2Input : public UActorComponent
+{
+	GENERATED_BODY()
+
+public:
+	UPixelStreaming2Input(const FObjectInitializer& ObjectInitializer);
+	void BeginPlay() override;
+	void EndPlay(const EEndPlayReason::Type EndPlayReason) override;
+
+	// The delegate which will be notified about a UI interaction.
+	DECLARE_DYNAMIC_MULTICAST_DELEGATE_OneParam(FOnInput, const FString&, Descriptor);
+	UPROPERTY(BlueprintAssignable, Category = "PixelStreaming2 Input")
+	FOnInput OnInputEvent;
+
+	/**
+	 * Send a response back to the source of the UI interactions.
+	 * @param Descriptor - A generic descriptor string.
+	 */
+	UFUNCTION(BlueprintCallable, Category = "PixelStreaming2 Input")
+	void SendPixelStreaming2Response(const FString& Descriptor);
+
+	/**
+	 * Helper function to extract a string field from a JSON descriptor of a
+	 * UI interaction given its field name.
+	 * The field name may be hierarchical, delimited by a period. For example,
+	 * to access the Width value of a Resolution command above you should use
+	 * "Resolution.Width" to get the width value.
+	 * @param Descriptor - The UI interaction JSON descriptor.
+	 * @param FieldName - The name of the field to look for in the JSON.
+	 * @param StringValue - The string value associated with the field name.
+	 * @param Success - True if the field exists in the JSON data.
+	 */
+	UFUNCTION(BlueprintPure, Category = "PixelStreaming2 Input")
+	void GetJsonStringValue(FString Descriptor, FString FieldName, FString& StringValue, bool& Success);
+
+	/**
+	 * Helper function to add a string field to a JSON descriptor. This produces
+	 * a new descriptor which may then be chained to add further string fields.
+	 * @param Descriptor - The initial JSON descriptor which may be blank initially.
+	 * @param FieldName - The name of the field to add to the JSON.
+	 * @param StringValue - The string value associated with the field name.
+	 * @param NewDescriptor - The JSON descriptor with the string field added.
+	 * @param Success - True if the string field could be added successfully.
+	 */
+	UFUNCTION(BlueprintPure, Category = "PixelStreaming2 Input")
+	void AddJsonStringValue(const FString& Descriptor, FString FieldName, FString StringValue, FString& NewDescriptor, bool& Success);
+};
diff --git a/Engine/Plugins/Media/PixelStreaming2/Source/PixelStreaming2/Internal/PixelStreaming2Common.h b/Engine/Plugins/Media/PixelStreaming2/Source/PixelStreaming2/Internal/PixelStreaming2Common.h
new file mode 100644
--- /dev/null
+++ b/Engine/Plugins/Media/PixelStreaming2/Source/PixelStreaming2/Internal/PixelStreaming2Common.h
@@ -0,0 +1,10 @@
+// Copyright Epic Games, Inc. All Rights Reserved.
+
+#pragma once
+
+#include "Blueprints/PixelStreaming2InputComponent.h"
+
+namespace UE::PixelStreaming2
+{
+	extern PIXELSTREAMING2_API TArray<UPixelStreaming2Input*> InputComponents;
+} // namespace UE::PixelStreaming2
diff --git a/Engine/Plugins/Media/PixelStreaming2/Source/PixelStreaming2/Private/Blueprints/PixelStreaming2Blueprints.cpp b/Engine/Plugins/Media/PixelStreaming2/Source/PixelStreaming2/Private/Blueprints/PixelStreaming2Blueprints.cpp
--- a/Engine/Plugins/Media/PixelStreaming2/Source/PixelStreaming2/Private/Blueprints/PixelStreaming2Blueprints.cpp
+++ b/Engine/Plugins/Media/PixelStreaming2/Source/PixelStreaming2/Private/Blueprints/PixelStreaming2Blueprints.cpp
@@ -6,7 +6,6 @@
 #include "Logging.h"
 #include "Misc/FileHelper.h"
 #include "PixelStreaming2AudioComponent.h"
-#include "PixelStreaming2InputComponent.h"
 
 TSharedPtr<IPixelStreaming2Streamer> GetStreamer(FString StreamerId)
 {
diff --git a/Engine/Plugins/Media/PixelStreaming2/Source/PixelStreaming2/Private/Blueprints/PixelStreaming2InputComponent.cpp b/Engine/Plugins/Media/PixelStreaming2/Source/PixelStreaming2/Private/Blueprints/PixelStreaming2InputComponent.cpp
--- a/Engine/Plugins/Media/PixelStreaming2/Source/PixelStreaming2/Private/Blueprints/PixelStreaming2InputComponent.cpp
+++ b/Engine/Plugins/Media/PixelStreaming2/Source/PixelStreaming2/Private/Blueprints/PixelStreaming2InputComponent.cpp
@@ -1,21 +1,17 @@
 // Copyright Epic Games, Inc. All Rights Reserved.
 
-#include "PixelStreaming2InputComponent.h"
-#include "IPixelStreaming2Module.h"
-#include "PixelStreaming2Module.h"
+#include "Blueprints/PixelStreaming2InputComponent.h"
+
 #include "Engine/Engine.h"
 #include "Engine/World.h"
-#include "Policies/CondensedJsonPrintPolicy.h"
-#include "Serialization/JsonReader.h"
-#include "Serialization/JsonSerializer.h"
-#include "GameFramework/GameUserSettings.h"
-#include "Logging.h"
+#include "IPixelStreaming2Module.h"
 #include "IPixelStreaming2Streamer.h"
+#include "Logging.h"
+#include "PixelStreaming2Common.h"
 #include "PixelStreaming2Utils.h"
 
 UPixelStreaming2Input::UPixelStreaming2Input(const FObjectInitializer& ObjectInitializer)
 	: Super(ObjectInitializer)
-	, PixelStreaming2Module((UE::PixelStreaming2::FPixelStreaming2Module*)UE::PixelStreaming2::FPixelStreaming2Module::GetModule())
 {
 	bAutoActivate = true;
 	PrimaryComponentTick.bCanEverTick = true;
@@ -26,30 +22,14 @@ void UPixelStreaming2Input::BeginPlay()
 {
 	Super::BeginPlay();
 
-	if (PixelStreaming2Module)
-	{
-		// When this component is initializing it registers itself with the Pixel Streaming module.
-		PixelStreaming2Module->AddInputComponent(this);
-	}
-	else
-	{
-		UE_LOG(LogPixelStreaming2, Warning, TEXT("Pixel Streaming input component not added because Pixel Streaming module is not loaded. This is expected on dedicated servers."));
-	}
+	UE::PixelStreaming2::InputComponents.Add(this);
 }
 
 void UPixelStreaming2Input::EndPlay(const EEndPlayReason::Type EndPlayReason)
 {
 	Super::EndPlay(EndPlayReason);
 
-	if (PixelStreaming2Module)
-	{
-		// When this component is destructing it unregisters itself with the Pixel Streaming module.
-		PixelStreaming2Module->RemoveInputComponent(this);
-	}
-	else
-	{
-		UE_LOG(LogPixelStreaming2, Warning, TEXT("Pixel Streaming input component not removed because Pixel Streaming module is not loaded. This is expected on dedicated servers."));
-	}
+	UE::PixelStreaming2::InputComponents.Remove(this);
 }
 
 void UPixelStreaming2Input::SendPixelStreaming2Response(const FString& Descriptor)
diff --git a/Engine/Plugins/Media/PixelStreaming2/Source/PixelStreaming2/Private/Blueprints/PixelStreaming2InputComponent.h b/Engine/Plugins/Media/PixelStreaming2/Source/PixelStreaming2/Private/Blueprints/PixelStreaming2InputComponent.h
deleted file mode 100644
diff --git a/Engine/Plugins/Media/PixelStreaming2/Source/PixelStreaming2/Private/PixelStreaming2Module.cpp b/Engine/Plugins/Media/PixelStreaming2/Source/PixelStreaming2/Private/PixelStreaming2Module.cpp
--- a/Engine/Plugins/Media/PixelStreaming2/Source/PixelStreaming2/Private/PixelStreaming2Module.cpp
+++ b/Engine/Plugins/Media/PixelStreaming2/Source/PixelStreaming2/Private/PixelStreaming2Module.cpp
@@ -6,6 +6,7 @@
 #include "Engine/Engine.h"
 #include "IPixelStreaming2Streamer.h"
 #include "Logging.h"
+#include "PixelStreaming2Common.h"
 #include "PixelStreaming2Delegates.h"
 #include "PixelStreaming2PluginSettings.h"
 #include "RHI.h"
@@ -18,6 +19,8 @@
 
 namespace UE::PixelStreaming2
 {
+	TArray<UPixelStreaming2Input*> InputComponents = {};
+
 	FPixelStreaming2Module* FPixelStreaming2Module::PixelStreaming2Module = nullptr;
 
 	FPixelStreaming2Module* FPixelStreaming2Module::GetModule()
@@ -200,8 +203,6 @@ namespace UE::PixelStreaming2
 		// Any time we create a new streamer, populate it's signalling server URL with whatever is in the ini, console or command line
 		NewStreamer->SetConnectionURL(UPixelStreaming2PluginSettings::CVarConnectionURL.GetValueOnAnyThread());
 
-		RegisterCustomHandlers(NewStreamer);
-
 		Streamers.Add(StreamerId, NewStreamer);
 
 		return NewStreamer;
@@ -294,106 +295,6 @@ namespace UE::PixelStreaming2
 		}
 	}
 
-	void FPixelStreaming2Module::AddInputComponent(UPixelStreaming2Input* InInputComponent)
-	{
-		InputComponents.Add(InInputComponent);
-	}
-
-	void FPixelStreaming2Module::RemoveInputComponent(UPixelStreaming2Input* InInputComponent)
-	{
-		InputComponents.Remove(InInputComponent);
-	}
-
-	const TArray<UPixelStreaming2Input*> FPixelStreaming2Module::GetInputComponents()
-	{
-		return InputComponents;
-	}
-
-	void FPixelStreaming2Module::RegisterCustomHandlers(TSharedPtr<IPixelStreaming2Streamer> Streamer)
-	{
-		if (TSharedPtr<IPixelStreaming2InputHandler> InputHandler = Streamer->GetInputHandler().Pin())
-		{
-			// Set Encoder.MinQP Legacy CVar
-			InputHandler->SetCommandHandler(TEXT("Encoder.MinQP"), [](FString PlayerId, FString Descriptor, FString MinQPString) {
-				int MinQP = FCString::Atoi(*MinQPString);
-				UPixelStreaming2PluginSettings::CVarEncoderMaxQuality->SetWithCurrentPriority(FMath::RoundToFloat(100.0f * (1.0f - (FMath::Clamp<int32>(MinQP, 0, 51) / 51.0f))));
-			});
-
-			// Set Encoder.MaxQP Legacy CVar
-			InputHandler->SetCommandHandler(TEXT("Encoder.MaxQP"), [](FString PlayerId, FString Descriptor, FString MaxQPString) {
-				int MaxQP = FCString::Atoi(*MaxQPString);
-				UPixelStreaming2PluginSettings::CVarEncoderMinQuality->SetWithCurrentPriority(FMath::RoundToFloat(100.0f * (1.0f - (FMath::Clamp<int32>(MaxQP, 0, 51) / 51.0f))));
-			});
-
-			// Set Encoder.MinQuality CVar
-			InputHandler->SetCommandHandler(TEXT("Encoder.MinQuality"), [](FString PlayerId, FString Descriptor, FString MinQualityString) {
-				int MinQuality = FCString::Atoi(*MinQualityString);
-				UPixelStreaming2PluginSettings::CVarEncoderMinQuality->SetWithCurrentPriority(FMath::Clamp<int32>(MinQuality, 0, 100));
-			});
-
-			// Set Encoder.MaxQuality CVar
-			InputHandler->SetCommandHandler(TEXT("Encoder.MaxQuality"), [](FString PlayerId, FString Descriptor, FString MaxQualityString) {
-				int MaxQuality = FCString::Atoi(*MaxQualityString);
-				UPixelStreaming2PluginSettings::CVarEncoderMaxQuality->SetWithCurrentPriority(FMath::Clamp<int32>(MaxQuality, 0, 100));
-			});
-
-			// Set WebRTC max FPS
-			InputHandler->SetCommandHandler(TEXT("WebRTC.Fps"), [](FString PlayerId, FString Descriptor, FString FPSString) {
-				int FPS = FCString::Atoi(*FPSString);
-				UPixelStreaming2PluginSettings::CVarWebRTCFps->SetWithCurrentPriority(FPS);
-			});
-
-			// Set MinBitrate
-			InputHandler->SetCommandHandler(TEXT("WebRTC.MinBitrate"), [InputHandler](FString PlayerId, FString Descriptor, FString MinBitrateString) {
-				if (InputHandler->IsElevated(PlayerId))
-				{
-					int MinBitrate = FCString::Atoi(*MinBitrateString);
-					UPixelStreaming2PluginSettings::CVarWebRTCMinBitrate->SetWithCurrentPriority(MinBitrate);
-				}
-			});
-
-			// Set MaxBitrate
-			InputHandler->SetCommandHandler(TEXT("WebRTC.MaxBitrate"), [InputHandler](FString PlayerId, FString Descriptor, FString MaxBitrateString) {
-				if (InputHandler->IsElevated(PlayerId))
-				{
-					int MaxBitrate = FCString::Atoi(*MaxBitrateString);
-					UPixelStreaming2PluginSettings::CVarWebRTCMaxBitrate->SetWithCurrentPriority(MaxBitrate);
-				}
-			});
-
-			InputHandler->RegisterMessageHandler(EPixelStreaming2ToStreamerMessage::UIInteraction, [this](FString PlayerId, FMemoryReader Ar) { HandleUIInteraction(Ar); });
-
-			// Handle special cases when the InputHandler itself wants to send a message out to all the peers.
-			// Some special cases include when virtual gamepads are connected and a controller id needs to be transmitted.
-			TWeakPtr<IPixelStreaming2Streamer> WeakStreamer = Streamer;
-			InputHandler->OnSendMessage.AddLambda([WeakStreamer](FString MessageName, FMemoryReader Ar) {
-				TSharedPtr<IPixelStreaming2Streamer> Streamer = WeakStreamer.Pin();
-				if (!Streamer)
-				{
-					return;
-				}
-				FString Descriptor;
-				Ar << Descriptor;
-				Streamer->SendAllPlayersMessage(MessageName, Descriptor);
-			});
-		}
-	}
-
-	void FPixelStreaming2Module::HandleUIInteraction(FMemoryReader Ar)
-	{
-		FString Res;
-		Res.GetCharArray().SetNumUninitialized(Ar.TotalSize() / 2 + 1);
-		Ar.Serialize(Res.GetCharArray().GetData(), Ar.TotalSize());
-
-		FString Descriptor = Res.Mid(1);
-
-		UE_LOG(LogPixelStreaming2, Verbose, TEXT("UIInteraction: %s"), *Descriptor);
-		for (UPixelStreaming2Input* InputComponent : InputComponents)
-		{
-			InputComponent->OnInputEvent.Broadcast(Descriptor);
-		}
-	}
-
 	void FPixelStreaming2Module::InitDefaultStreamer()
 	{
 		UE_LOGFMT(LogPixelStreaming2, Log, "Initializing default streamer. ID: [{0}], Type: [{1}]", GetDefaultStreamerID(), UPixelStreaming2PluginSettings::CVarDefaultStreamerType.GetValueOnAnyThread());
diff --git a/Engine/Plugins/Media/PixelStreaming2/Source/PixelStreaming2/Private/PixelStreaming2Module.h b/Engine/Plugins/Media/PixelStreaming2/Source/PixelStreaming2/Private/PixelStreaming2Module.h
--- a/Engine/Plugins/Media/PixelStreaming2/Source/PixelStreaming2/Private/PixelStreaming2Module.h
+++ b/Engine/Plugins/Media/PixelStreaming2/Source/PixelStreaming2/Private/PixelStreaming2Module.h
@@ -34,11 +34,6 @@ namespace UE::PixelStreaming2
 		virtual FString									  GetDefaultConnectionURL() override;
 		virtual void									  ForEachStreamer(const TFunction<void(TSharedPtr<IPixelStreaming2Streamer>)>& Func) override;
 
-		// These are staying on the module at the moment as theres no way of the BPs knowing which streamer they are relevant to
-		void								 AddInputComponent(UPixelStreaming2Input* InInputComponent);
-		void								 RemoveInputComponent(UPixelStreaming2Input* InInputComponent);
-		const TArray<UPixelStreaming2Input*> GetInputComponents();
-
 	private:
 		// Begin IModuleInterface
 		void StartupModule() override;
@@ -46,17 +41,12 @@ namespace UE::PixelStreaming2
 		// End IModuleInterface
 
 	private:
-		void RegisterCustomHandlers(TSharedPtr<IPixelStreaming2Streamer> Streamer);
-		void HandleUIInteraction(FMemoryReader Ar);
-
 		void InitDefaultStreamer();
 
 	private:
 		bool		bModuleReady = false;
 		FReadyEvent ReadyEvent;
 
-		TArray<UPixelStreaming2Input*> InputComponents;
-
 		// FPixelStreamingThread must exist before any AudioTask and AudioMixingCapturer(Which contains a audio task) to ensure it is destroyed last
 		TSharedPtr<class FPixelStreamingThread> PixelStreamingThread;
 
diff --git a/Engine/Plugins/Media/PixelStreaming2/Source/PixelStreaming2Input/Internal/InputHandler.h b/Engine/Plugins/Media/PixelStreaming2/Source/PixelStreaming2Input/Internal/InputHandler.h
--- a/Engine/Plugins/Media/PixelStreaming2/Source/PixelStreaming2Input/Internal/InputHandler.h
+++ b/Engine/Plugins/Media/PixelStreaming2/Source/PixelStreaming2Input/Internal/InputHandler.h
@@ -95,6 +95,7 @@ namespace UE::PixelStreaming2Input
 		virtual bool OnXRSystem(EPixelStreaming2XRSystem System) override;
 
 	protected:
+		FVector2D	ConvertToNormalizedScreenLocation(FVector2D Point);
 		FIntPoint	ConvertFromNormalizedScreenLocation(const FVector2D& ScreenLocation, bool bIncludeOffset = true);
 		FWidgetPath FindRoutingMessageWidget(const FVector2D& Location) const;
 		FKey		TranslateMouseButtonToKey(const EMouseButtons::Type Button);
@@ -147,6 +148,12 @@ namespace UE::PixelStreaming2Input
 
 		void FindFocusedWidget();
 
+		/**
+		 * Create an artificial mouse 'movement' to allow widgets to focus under a
+		 * static mouse cursor. This will not actually change the cursor position.
+		 */
+		void SynthesizeMouseMove() const;
+
 		TArray<FKey> FilteredKeys;
 		void		 OnFilteredKeysChanged(IConsoleVariable* Var);
 		bool		 FilterKey(const FKey& Key);
@@ -189,6 +196,25 @@ namespace UE::PixelStreaming2Input
 		 */
 		const size_t MessageHeaderOffset = 1;
 
+		/**
+		 * Touch only. Location of the focused UI widget. If no UI widget is focused
+		 * then this has the UnfocusedPos value.
+		 */
+		FVector2D FocusedPos;
+
+		/**
+		 * Touch only. A special position which indicates that no UI widget is
+		 * focused.
+		 */
+		const FVector2D UnfocusedPos = FVector2D(-1.f, -1.f);
+
+		/**
+		 * Whether an artificial mouse 'movement' should be created the next time
+		 * a mouse button down event occurs. This allows a widget to be focused
+		 * immediately when the browser window is focused.
+		 */
+		bool bSynthesizeMouseMoveForNextMouseDown = false;
+
 		struct FPixelStreaming2XRController
 		{
 		public:
diff --git a/Engine/Plugins/Media/PixelStreaming2/Source/PixelStreaming2Input/Private/InputHandler.cpp b/Engine/Plugins/Media/PixelStreaming2/Source/PixelStreaming2Input/Private/InputHandler.cpp
--- a/Engine/Plugins/Media/PixelStreaming2/Source/PixelStreaming2Input/Private/InputHandler.cpp
+++ b/Engine/Plugins/Media/PixelStreaming2/Source/PixelStreaming2Input/Private/InputHandler.cpp
@@ -43,6 +43,7 @@ namespace UE::PixelStreaming2Input
 		, bIsMouseActive(false)
 		, MessageHandler(FSlateApplication::Get().GetPlatformApplication()->GetMessageHandler())
 		, PixelStreamerApplicationWrapper(MakeShared<FPixelStreaming2ApplicationWrapper>(FSlateApplication::Get().GetPlatformApplication()))
+		, FocusedPos(UnfocusedPos)
 
 	{
 		// Register this input handler as an IMotionController. The module handles the registering as an IInputDevice
@@ -581,6 +582,13 @@ namespace UE::PixelStreaming2Input
 		// Force window focus
 		SlateApplication.ProcessApplicationActivationEvent(true);
 
+		if (bSynthesizeMouseMoveForNextMouseDown)
+		{
+			SynthesizeMouseMove();
+			bSynthesizeMouseMoveForNextMouseDown = false;
+		}
+
+		bool bHandled = false;
 		if (InputType == EPixelStreaming2InputType::RouteToWidget)
 		{
 			FWidgetPath WidgetPath = FindRoutingMessageWidget(ScreenLocation);
@@ -601,15 +609,19 @@ namespace UE::PixelStreaming2Input
 					0,
 					SlateApplication.GetPlatformApplication()->GetModifierKeys());
 
-				return SlateApplication.RoutePointerDownEvent(WidgetPath, MouseEvent).IsEventHandled();
+				bHandled = SlateApplication.RoutePointerDownEvent(WidgetPath, MouseEvent).IsEventHandled();
 			}
 		}
 		else if (InputType == EPixelStreaming2InputType::RouteToWindow)
 		{
-			return MessageHandler->OnMouseDown(PixelStreamerApplicationWrapper->GetWindowUnderCursor(), Button, ScreenLocation);
+			bHandled = MessageHandler->OnMouseDown(PixelStreamerApplicationWrapper->GetWindowUnderCursor(), Button, ScreenLocation);
 		}
 
-		return false;
+		// The browser may be faking a mouse when touching so it will send
+		// over a mouse down event.
+		FindFocusedWidget();
+
+		return bHandled;
 	}
 
 	bool FPixelStreaming2InputHandler::OnMouseMove(FIntPoint ScreenLocation, FIntPoint Delta)
@@ -833,6 +845,53 @@ namespace UE::PixelStreaming2Input
 		return !ElevatedCheck || ElevatedCheck(Id);
 	}
 
+	FVector2D FPixelStreaming2InputHandler::ConvertToNormalizedScreenLocation(FVector2D Pos)
+	{
+		FVector2D NormalizedLocation = FVector2D::ZeroVector;
+
+		TSharedPtr<SWindow> ApplicationWindow = TargetWindow.Pin();
+		if (ApplicationWindow.IsValid())
+		{
+			FVector2D WindowOrigin = ApplicationWindow->GetPositionInScreen();
+			if (TargetViewport.IsValid())
+			{
+				TSharedPtr<SViewport> ViewportWidget = TargetViewport.Pin();
+
+				if (ViewportWidget.IsValid())
+				{
+					FGeometry InnerWindowGeometry = ApplicationWindow->GetWindowGeometryInWindow();
+
+					// Find the widget path relative to the window
+					FArrangedChildren JustWindow(EVisibility::Visible);
+					JustWindow.AddWidget(FArrangedWidget(ApplicationWindow.ToSharedRef(), InnerWindowGeometry));
+
+					FWidgetPath PathToWidget(ApplicationWindow.ToSharedRef(), JustWindow);
+					if (PathToWidget.ExtendPathTo(FWidgetMatcher(ViewportWidget.ToSharedRef()), EVisibility::Visible))
+					{
+						FArrangedWidget ArrangedWidget = PathToWidget.FindArrangedWidget(ViewportWidget.ToSharedRef()).Get(FArrangedWidget::GetNullWidget());
+						FVector2D		WindowClientSize = ArrangedWidget.Geometry.GetAbsoluteSize();
+
+						NormalizedLocation = FVector2D(Pos / WindowClientSize);
+					}
+				}
+			}
+			else
+			{
+				FVector2D SizeInScreen = ApplicationWindow->GetSizeInScreen();
+				NormalizedLocation = FVector2D(Pos / SizeInScreen);
+			}
+		}
+		else if (TSharedPtr<FIntRect> ScreenRectPtr = TargetScreenRect.Pin())
+		{
+			FIntRect  ScreenRect = *ScreenRectPtr;
+			FIntPoint SizeInScreen = ScreenRect.Max - ScreenRect.Min;
+			NormalizedLocation = Pos / SizeInScreen;
+		}
+
+		NormalizedLocation *= uint16_MAX;
+		return NormalizedLocation;
+	}
+
 	FIntPoint FPixelStreaming2InputHandler::ConvertFromNormalizedScreenLocation(const FVector2D& ScreenLocation, bool bIncludeOffset)
 	{
 		FIntPoint OutVector((int32)ScreenLocation.X, (int32)ScreenLocation.Y);
@@ -1006,6 +1065,14 @@ namespace UE::PixelStreaming2Input
 		}
 	}
 
+	void FPixelStreaming2InputHandler::SynthesizeMouseMove() const
+	{
+		// Move the mouse back and forth so the net result does not result in
+		// moving the cursor.
+		MessageHandler->OnRawMouseMove(1, 0);
+		MessageHandler->OnRawMouseMove(-1, 0);
+	}
+
 	FKey FPixelStreaming2InputHandler::TranslateMouseButtonToKey(const EMouseButtons::Type Button)
 	{
 		FKey Key = EKeys::Invalid;
@@ -1040,6 +1107,7 @@ namespace UE::PixelStreaming2Input
 			static FName SEditableTextType(TEXT("SEditableText"));
 			static FName SMultiLineEditableTextType(TEXT("SMultiLineEditableText"));
 			bool		 bEditable = FocusedWidget && (FocusedWidget->GetType() == SEditableTextType || FocusedWidget->GetType() == SMultiLineEditableTextType);
+
 			if (bEditable)
 			{
 				if (FocusedWidget->GetType() == TEXT("SEditableText"))
@@ -1054,19 +1122,15 @@ namespace UE::PixelStreaming2Input
 				}
 			}
 
-			// Tell the browser that the focus has changed.
-			TSharedPtr<FJsonObject> JsonObject = MakeShareable(new FJsonObject);
-			JsonObject->SetStringField(TEXT("command"), TEXT("onScreenKeyboard"));
-			JsonObject->SetBoolField(TEXT("showOnScreenKeyboard"), bEditable);
-
+			FVector2D Pos = UnfocusedPos;
 			if (bEditable)
 			{
-				FVector2D			NormalizedLocation = FVector2D::ZeroVector;
+				Pos = FocusedWidget->GetCachedGeometry().GetAbsolutePosition();
+
 				TSharedPtr<SWindow> ApplicationWindow = TargetWindow.Pin();
 				if (ApplicationWindow.IsValid())
 				{
 					FVector2D WindowOrigin = ApplicationWindow->GetPositionInScreen();
-					FVector2D Pos = FocusedWidget->GetCachedGeometry().GetAbsolutePosition();
 					if (TargetViewport.IsValid())
 					{
 						TSharedPtr<SViewport> ViewportWidget = TargetViewport.Pin();
@@ -1088,64 +1152,65 @@ namespace UE::PixelStreaming2Input
 								FVector2D WindowClientSize = ArrangedWidget.Geometry.GetAbsoluteSize();
 
 								Pos = Pos - WindowClientOffset;
-								NormalizedLocation = FVector2D(Pos / WindowClientSize);
 							}
 						}
 					}
-					else
-					{
-						FVector2D SizeInScreen = ApplicationWindow->GetSizeInScreen();
-						NormalizedLocation = FVector2D(Pos / SizeInScreen);
-					}
-				}
-				else if (TSharedPtr<FIntRect> ScreenRectPtr = TargetScreenRect.Pin())
-				{
-					FIntRect  ScreenRect = *ScreenRectPtr;
-					FIntPoint SizeInScreen = ScreenRect.Max - ScreenRect.Min;
-					NormalizedLocation = FocusedWidget->GetCachedGeometry().GetAbsolutePosition() / SizeInScreen;
 				}
+			}
 
-				NormalizedLocation *= uint16_MAX;
-				// ConvertToNormalizedScreenLocation(Pos, NormalizedLocation);
-				JsonObject->SetNumberField(TEXT("x"), static_cast<uint16>(NormalizedLocation.X));
-				JsonObject->SetNumberField(TEXT("y"), static_cast<uint16>(NormalizedLocation.Y));
+			if (Pos != FocusedPos)
+			{
+				FocusedPos = Pos;
 
-				FText TextboxContents;
-				if (FocusedWidget->GetType() == TEXT("SEditableText"))
-				{
-					SEditableText* TextBox = static_cast<SEditableText*>(FocusedWidget.Get());
-					TextboxContents = TextBox->GetText();
-				}
-				else if (FocusedWidget->GetType() == TEXT("SMultiLineEditableText"))
+				// Tell the browser that the focus has changed.
+				TSharedPtr<FJsonObject> JsonObject = MakeShareable(new FJsonObject);
+				JsonObject->SetStringField(TEXT("command"), TEXT("onScreenKeyboard"));
+				JsonObject->SetBoolField(TEXT("showOnScreenKeyboard"), bEditable);
+
+				if (bEditable)
 				{
-					SMultiLineEditableText* TextBox = static_cast<SMultiLineEditableText*>(FocusedWidget.Get());
-					TextboxContents = TextBox->GetText();
+					FVector2D NormalizedLocation = ConvertToNormalizedScreenLocation(Pos);
+
+					JsonObject->SetNumberField(TEXT("x"), static_cast<uint16>(NormalizedLocation.X));
+					JsonObject->SetNumberField(TEXT("y"), static_cast<uint16>(NormalizedLocation.Y));
+
+					FText TextboxContents;
+					if (FocusedWidget->GetType() == TEXT("SEditableText"))
+					{
+						SEditableText* TextBox = static_cast<SEditableText*>(FocusedWidget.Get());
+						TextboxContents = TextBox->GetText();
+					}
+					else if (FocusedWidget->GetType() == TEXT("SMultiLineEditableText"))
+					{
+						SMultiLineEditableText* TextBox = static_cast<SMultiLineEditableText*>(FocusedWidget.Get());
+						TextboxContents = TextBox->GetText();
+					}
+
+					JsonObject->SetStringField(TEXT("contents"), TextboxContents.ToString());
 				}
 
-				JsonObject->SetStringField(TEXT("contents"), TextboxContents.ToString());
+				FString															 Descriptor;
+				TSharedRef<TJsonWriter<TCHAR, TCondensedJsonPrintPolicy<TCHAR>>> JsonWriter = TJsonWriterFactory<TCHAR, TCondensedJsonPrintPolicy<TCHAR>>::Create(&Descriptor);
+				FJsonSerializer::Serialize(JsonObject.ToSharedRef(), JsonWriter);
+
+				FBufferArchive Buffer;
+				Buffer << Descriptor;
+				TArray<uint8> Data(Buffer.GetData(), Buffer.Num());
+
+				/**
+				 * Send the following JSON string in a "Command" message to all players
+				 * {
+				 *	"command": "onScreenKeyboard",
+				 *	"showOnScreenKeyboard": "true", //or false
+				 *   "x": 1, //some uint16
+				 *   "y": 1, //some uint16
+				 *   "contents": "text box contents" // whatever text the textbox has in it
+				 * }
+				 */
+
+				// Specific implementation for this method is handled per streamer
+				OnSendMessage.Broadcast(EPixelStreaming2FromStreamerMessage::Command, FMemoryReader(Data));
 			}
-
-			FString															 Descriptor;
-			TSharedRef<TJsonWriter<TCHAR, TCondensedJsonPrintPolicy<TCHAR>>> JsonWriter = TJsonWriterFactory<TCHAR, TCondensedJsonPrintPolicy<TCHAR>>::Create(&Descriptor);
-			FJsonSerializer::Serialize(JsonObject.ToSharedRef(), JsonWriter);
-
-			FBufferArchive Buffer;
-			Buffer << Descriptor;
-			TArray<uint8> Data(Buffer.GetData(), Buffer.Num());
-
-			/**
-			 * Send the following JSON string in a "Command" message to all players
-			 * {
-			 *	"command": "onScreenKeyboard",
-			 *	"showOnScreenKeyboard": "true", //or false
-			 *   "x": 1, //some uint16
-			 *   "y": 1, //some uint16
-			 *   "contents": "text box contents" // whatever text the textbox has in it
-			 * }
-			 */
-
-			// Specific implementation for this method is handled per streamer
-			OnSendMessage.Broadcast(EPixelStreaming2FromStreamerMessage::Command, FMemoryReader(Data));
 		});
 	}
 
diff --git a/Engine/Plugins/Media/PixelStreaming2/Source/PixelStreaming2RTC/Internal/EpicRtcDataTrack.h b/Engine/Plugins/Media/PixelStreaming2/Source/PixelStreaming2RTC/Internal/EpicRtcDataTrack.h
--- a/Engine/Plugins/Media/PixelStreaming2/Source/PixelStreaming2RTC/Internal/EpicRtcDataTrack.h
+++ b/Engine/Plugins/Media/PixelStreaming2/Source/PixelStreaming2RTC/Internal/EpicRtcDataTrack.h
@@ -31,6 +31,18 @@ namespace UE::PixelStreaming2
 		return *Value;
 	}
 
+	template <typename T>
+	inline size_t ValueSize(TArray<T>&& Value)
+	{
+		return Value.Num() * sizeof(T);
+	}
+
+	template <typename T>
+	inline const void* ValueLoc(TArray<T>&& Value)
+	{
+		return Value.GetData();
+	}
+
 	struct PIXELSTREAMING2RTC_API FBufferBuilder
 	{
 		TArray<uint8> Buffer;
diff --git a/Engine/Plugins/Media/PixelStreaming2/Source/PixelStreaming2RTC/Private/EpicRtcStreamer.cpp b/Engine/Plugins/Media/PixelStreaming2/Source/PixelStreaming2RTC/Private/EpicRtcStreamer.cpp
--- a/Engine/Plugins/Media/PixelStreaming2/Source/PixelStreaming2RTC/Private/EpicRtcStreamer.cpp
+++ b/Engine/Plugins/Media/PixelStreaming2/Source/PixelStreaming2RTC/Private/EpicRtcStreamer.cpp
@@ -31,6 +31,7 @@
 #include "VideoProducerBackBuffer.h"
 #include "VideoProducerMediaCapture.h"
 #include "RTCInputHandler.h"
+#include "PixelStreaming2Common.h"
 
 namespace UE::PixelStreaming2
 {
@@ -69,10 +70,96 @@ namespace UE::PixelStreaming2
 	void FEpicRtcStreamer::Initialize()
 	{
 		TSharedPtr<FEpicRtcStreamer> Streamer = AsShared();
+		TWeakPtr<FEpicRtcStreamer>	 WeakStreamer = Streamer;
 
 		InputHandler->GetToStreamerProtocol()->OnProtocolUpdated().AddSP(Streamer.ToSharedRef(), &FEpicRtcStreamer::OnProtocolUpdated);
 		InputHandler->GetFromStreamerProtocol()->OnProtocolUpdated().AddSP(Streamer.ToSharedRef(), &FEpicRtcStreamer::OnProtocolUpdated);
 
+		// Set Encoder.MinQP Legacy CVar
+		InputHandler->SetCommandHandler(TEXT("Encoder.MinQP"), [](FString PlayerId, FString Descriptor, FString MinQPString) {
+			int MinQP = FCString::Atoi(*MinQPString);
+			UPixelStreaming2PluginSettings::CVarEncoderMaxQuality->SetWithCurrentPriority(FMath::RoundToFloat(100.0f * (1.0f - (FMath::Clamp<int32>(MinQP, 0, 51) / 51.0f))));
+		});
+
+		// Set Encoder.MaxQP Legacy CVar
+		InputHandler->SetCommandHandler(TEXT("Encoder.MaxQP"), [](FString PlayerId, FString Descriptor, FString MaxQPString) {
+			int MaxQP = FCString::Atoi(*MaxQPString);
+			UPixelStreaming2PluginSettings::CVarEncoderMinQuality->SetWithCurrentPriority(FMath::RoundToFloat(100.0f * (1.0f - (FMath::Clamp<int32>(MaxQP, 0, 51) / 51.0f))));
+		});
+
+		// Set Encoder.MinQuality CVar
+		InputHandler->SetCommandHandler(TEXT("Encoder.MinQuality"), [](FString PlayerId, FString Descriptor, FString MinQualityString) {
+			int MinQuality = FCString::Atoi(*MinQualityString);
+			UPixelStreaming2PluginSettings::CVarEncoderMinQuality->SetWithCurrentPriority(FMath::Clamp<int32>(MinQuality, 0, 100));
+		});
+
+		// Set Encoder.MaxQuality CVar
+		InputHandler->SetCommandHandler(TEXT("Encoder.MaxQuality"), [](FString PlayerId, FString Descriptor, FString MaxQualityString) {
+			int MaxQuality = FCString::Atoi(*MaxQualityString);
+			UPixelStreaming2PluginSettings::CVarEncoderMaxQuality->SetWithCurrentPriority(FMath::Clamp<int32>(MaxQuality, 0, 100));
+		});
+
+		// Set WebRTC max FPS
+		InputHandler->SetCommandHandler(TEXT("WebRTC.Fps"), [](FString PlayerId, FString Descriptor, FString FPSString) {
+			int FPS = FCString::Atoi(*FPSString);
+			UPixelStreaming2PluginSettings::CVarWebRTCFps->SetWithCurrentPriority(FPS);
+		});
+
+		// Set MinBitrate
+		InputHandler->SetCommandHandler(TEXT("WebRTC.MinBitrate"), [WeakStreamer](FString PlayerId, FString Descriptor, FString MinBitrateString) {
+			// This check pattern is kind of verbose, but as the messages are enqueued on a different thread to where they were added
+			// we need to make sure that both the streamer and input handler are still alive when we process the command
+			TSharedPtr<FEpicRtcStreamer> Streamer = WeakStreamer.Pin();
+			if (!Streamer)
+			{
+				return;
+			}
+			TSharedPtr<IPixelStreaming2InputHandler> InputHandler = Streamer->GetInputHandler().Pin();
+			if (!InputHandler)
+			{
+				return;
+			}
+
+			if (InputHandler->IsElevated(PlayerId))
+			{
+				int MinBitrate = FCString::Atoi(*MinBitrateString);
+				UPixelStreaming2PluginSettings::CVarWebRTCMinBitrate->SetWithCurrentPriority(MinBitrate);
+			}
+		});
+
+		// Set MaxBitrate
+		InputHandler->SetCommandHandler(TEXT("WebRTC.MaxBitrate"), [WeakStreamer](FString PlayerId, FString Descriptor, FString MaxBitrateString) {
+			TSharedPtr<FEpicRtcStreamer> Streamer = WeakStreamer.Pin();
+			if (!Streamer)
+			{
+				return;
+			}
+			TSharedPtr<IPixelStreaming2InputHandler> InputHandler = Streamer->GetInputHandler().Pin();
+			if (!InputHandler)
+			{
+				return;
+			}
+
+			if (InputHandler->IsElevated(PlayerId))
+			{
+				int MaxBitrate = FCString::Atoi(*MaxBitrateString);
+				UPixelStreaming2PluginSettings::CVarWebRTCMaxBitrate->SetWithCurrentPriority(MaxBitrate);
+			}
+		});
+
+		InputHandler->RegisterMessageHandler(EPixelStreaming2ToStreamerMessage::UIInteraction, [WeakStreamer](FString PlayerId, FMemoryReader Ar) {
+			TSharedPtr<FEpicRtcStreamer> Streamer = WeakStreamer.Pin();
+			if (!Streamer)
+			{
+				return;
+			}
+			Streamer->OnUIInteraction(Ar);
+		});
+
+		// Handle special cases when the InputHandler itself wants to send a message out to all the peers.
+		// Some special cases include when virtual gamepads are connected and a controller id needs to be transmitted.
+		InputHandler->OnSendMessage.AddSP(Streamer.ToSharedRef(), &FEpicRtcStreamer::OnSendMessage);
+
 		Streamer->EpicRtcManager->OnSessionStateUpdate.AddSP(Streamer.ToSharedRef(), &FEpicRtcStreamer::OnSessionStateUpdate);
 		Streamer->EpicRtcManager->OnSessionErrorUpdate.AddSP(Streamer.ToSharedRef(), &FEpicRtcStreamer::OnSessionErrorUpdate);
 		Streamer->EpicRtcManager->OnSessionRoomsAvailableUpdate.AddSP(Streamer.ToSharedRef(), &FEpicRtcStreamer::OnSessionRoomsAvailableUpdate);
@@ -278,6 +365,28 @@ namespace UE::PixelStreaming2
 		PlayerContext->StatsCollector->Process(ConnectionStats);
 	}
 
+	void FEpicRtcStreamer::OnUIInteraction(FMemoryReader Ar)
+	{
+		FString Res;
+		Res.GetCharArray().SetNumUninitialized(Ar.TotalSize() / 2 + 1);
+		Ar.Serialize(Res.GetCharArray().GetData(), Ar.TotalSize());
+
+		FString Descriptor = Res.Mid(1);
+
+		UE_LOGFMT(LogPixelStreaming2RTC, Verbose, "FEpicRtcStreamer[\"{0}\"]::OnUIInteraction({1})", StreamerId, Descriptor);
+		for (UPixelStreaming2Input* InputComponent : InputComponents)
+		{
+			InputComponent->OnInputEvent.Broadcast(Descriptor);
+		}
+	}
+
+	void FEpicRtcStreamer::OnSendMessage(FString MessageName, FMemoryReader Ar)
+	{
+		FString Descriptor;
+		Ar << Descriptor;
+		SendAllPlayersMessage(MessageName, Descriptor);
+	}
+
 	IPixelStreaming2Streamer::FPreConnectionEvent& FEpicRtcStreamer::OnPreConnection()
 	{
 		return StreamingPreConnectionEvent;
diff --git a/Engine/Plugins/Media/PixelStreaming2/Source/PixelStreaming2RTC/Private/EpicRtcStreamer.h b/Engine/Plugins/Media/PixelStreaming2/Source/PixelStreaming2RTC/Private/EpicRtcStreamer.h
--- a/Engine/Plugins/Media/PixelStreaming2/Source/PixelStreaming2RTC/Private/EpicRtcStreamer.h
+++ b/Engine/Plugins/Media/PixelStreaming2/Source/PixelStreaming2RTC/Private/EpicRtcStreamer.h
@@ -97,6 +97,8 @@ namespace UE::PixelStreaming2
 		void RemoveSession(bool bDisconnect);
 		void RemoveRoom();
 		void OnStatsReady(const FString& PlayerId, const EpicRtcConnectionStats& ConnectionStats);
+		void OnUIInteraction(FMemoryReader Ar);
+		void OnSendMessage(FString MessageName, FMemoryReader Ar);
 
 	private:
 		FString StreamerId;
