[PS2] Fix: Always call start streaming on the DefaultStreamer

diff --git a/Engine/Plugins/Media/PixelStreaming2/Source/PixelStreaming2/Private/PixelStreaming2Module.cpp b/Engine/Plugins/Media/PixelStreaming2/Source/PixelStreaming2/Private/PixelStreaming2Module.cpp
--- a/Engine/Plugins/Media/PixelStreaming2/Source/PixelStreaming2/Private/PixelStreaming2Module.cpp
+++ b/Engine/Plugins/Media/PixelStreaming2/Source/PixelStreaming2/Private/PixelStreaming2Module.cpp
@@ -128,11 +128,7 @@ namespace UE::PixelStreaming2
 				// We don't want to create the default streamer when using the editor
 				InitDefaultStreamer();
 
-				// The user has specified a URL on the command line or ini meaning their intention is to start streaming immediately
-				if (!DefaultStreamer->GetConnectionURL().IsEmpty())
-				{
-					DefaultStreamer->StartStreaming();
-				}
+				DefaultStreamer->StartStreaming();
 			}
 		});
 	}
