[PS2] QOL: Don't set streamer connection url during streamer creation if the cvar is empty. Some users may wish to pass in the connection url using a method other than the cvar/command line

diff --git a/Engine/Plugins/Media/PixelStreaming2/Source/PixelStreaming2/Private/PixelStreaming2Module.cpp b/Engine/Plugins/Media/PixelStreaming2/Source/PixelStreaming2/Private/PixelStreaming2Module.cpp
--- a/Engine/Plugins/Media/PixelStreaming2/Source/PixelStreaming2/Private/PixelStreaming2Module.cpp
+++ b/Engine/Plugins/Media/PixelStreaming2/Source/PixelStreaming2/Private/PixelStreaming2Module.cpp
@@ -197,7 +197,10 @@ namespace UE::PixelStreaming2
 		TSharedPtr<IPixelStreaming2Streamer> NewStreamer = Factory->CreateNewStreamer(StreamerId);
 		NewStreamer->Initialize();
 		// Any time we create a new streamer, populate it's signalling server URL with whatever is in the ini, console or command line
-		NewStreamer->SetConnectionURL(UPixelStreaming2PluginSettings::CVarConnectionURL.GetValueOnAnyThread());
+		if(!UPixelStreaming2PluginSettings::CVarConnectionURL.GetValueOnAnyThread().IsEmpty())
+		{
+			NewStreamer->SetConnectionURL(UPixelStreaming2PluginSettings::CVarConnectionURL.GetValueOnAnyThread());
+		}
 
 		Streamers.Add(StreamerId, NewStreamer);
 
