[PS2] Fix black stream when resizing with multiple streamers connected.

When multiple streamers are connected and the stream resizes, the encoder ends up waiting for a timeout for all but one stream. This fix lowers that timeout.

diff --git a/Engine/Plugins/Media/PixelStreaming2/Source/PixelStreaming2RTC/Private/EpicRtcVideoBufferMultiFormat.cpp b/Engine/Plugins/Media/PixelStreaming2/Source/PixelStreaming2RTC/Private/EpicRtcVideoBufferMultiFormat.cpp
--- a/Engine/Plugins/Media/PixelStreaming2/Source/PixelStreaming2RTC/Private/EpicRtcVideoBufferMultiFormat.cpp
+++ b/Engine/Plugins/Media/PixelStreaming2/Source/PixelStreaming2RTC/Private/EpicRtcVideoBufferMultiFormat.cpp
@@ -62,8 +62,8 @@ namespace UE::PixelStreaming2
 		{
 			return nullptr;
 		}
-
-		TSharedPtr<IPixelCaptureOutputFrame> Frame = FrameCapturer->WaitForFormat(Format, LayerIndex);
+		constexpr uint32 WaitTimeMS = 500;
+		TSharedPtr<IPixelCaptureOutputFrame> Frame = FrameCapturer->WaitForFormat(Format, LayerIndex, WaitTimeMS);
 		CachedFormat.Add(Format, Frame);
 		return Frame.Get();
 	}
