[PS2] Fix RegisterMessageHandler XRSystem whcich was mistakenly set to "EPixelStreaming2XRSystem

diff --git a/Engine/Plugins/Media/PixelStreaming2/Source/PixelStreaming2RTC/Private/RTCInputHandler.cpp b/Engine/Plugins/Media/PixelStreaming2/Source/PixelStreaming2RTC/Private/RTCInputHandler.cpp
--- a/Engine/Plugins/Media/PixelStreaming2/Source/PixelStreaming2RTC/Private/RTCInputHandler.cpp
+++ b/Engine/Plugins/Media/PixelStreaming2/Source/PixelStreaming2RTC/Private/RTCInputHandler.cpp
@@ -60,7 +60,7 @@ namespace UE::PixelStreaming2
 		RegisterMessageHandler("XRButtonTouchReleased", [this](FString SourceId, FMemoryReader Ar) { HandleOnXRButtonTouchReleased(Ar); });
 		RegisterMessageHandler("XRButtonReleased", [this](FString SourceId, FMemoryReader Ar) { HandleOnXRButtonReleased(Ar); });
 		RegisterMessageHandler("XRAnalog", [this](FString SourceId, FMemoryReader Ar) { HandleOnXRAnalog(Ar); });
-		RegisterMessageHandler("EPixelStreaming2XRSystem", [this](FString SourceId, FMemoryReader Ar) { HandleOnXRSystem(Ar); });
+		RegisterMessageHandler("XRSystem", [this](FString SourceId, FMemoryReader Ar) { HandleOnXRSystem(Ar); });
 
 		RegisterMessageHandler("Command", [this](FString SourceId, FMemoryReader Ar) { HandleOnCommand(SourceId, Ar); });
 		RegisterMessageHandler("UIInteraction", [this](FString SourceId, FMemoryReader Ar) { HandleUIInteraction(Ar); });
