[PS2] QOL: Expose RTCStatCollector for other internal use cases

diff --git a/Engine/Plugins/Media/PixelStreaming2/Source/PixelStreaming2/Private/Blueprints/PixelStreaming2Blueprints.cpp b/Engine/Plugins/Media/PixelStreaming2/Source/PixelStreaming2/Private/Blueprints/PixelStreaming2Blueprints.cpp
--- a/Engine/Plugins/Media/PixelStreaming2/Source/PixelStreaming2/Private/Blueprints/PixelStreaming2Blueprints.cpp
+++ b/Engine/Plugins/Media/PixelStreaming2/Source/PixelStreaming2/Private/Blueprints/PixelStreaming2Blueprints.cpp
@@ -6,6 +6,7 @@
 #include "Logging.h"
 #include "Misc/FileHelper.h"
 #include "PixelStreaming2AudioComponent.h"
+#include "PixelStreaming2Utils.h"
 
 TSharedPtr<IPixelStreaming2Streamer> GetStreamer(FString StreamerId)
 {
@@ -21,6 +22,20 @@ TSharedPtr<IPixelStreaming2Streamer> GetStreamer(FString StreamerId)
 	}
 }
 
+void UPixelStreaming2Blueprints::SendResponse(FString Descriptor)
+{
+	StreamerSendResponse(GetDefaultStreamerID(), Descriptor);
+}
+
+	
+void UPixelStreaming2Blueprints::StreamerSendResponse(FString StreamerId, FString Descriptor)
+{
+	if (TSharedPtr<IPixelStreaming2Streamer> Streamer = GetStreamer(StreamerId))
+	{
+		Streamer->SendAllPlayersMessage(TEXT("Response"), Descriptor);
+	}
+}
+
 void UPixelStreaming2Blueprints::SendFileAsByteArray(TArray<uint8> ByteArray, FString MimeType, FString FileExtension)
 {
 	StreamerSendFileAsByteArray(GetDefaultStreamerID(), ByteArray, MimeType, FileExtension);
@@ -132,6 +147,16 @@ FString UPixelStreaming2Blueprints::GetDefaultStreamerID()
 	return IPixelStreaming2Module::Get().GetDefaultStreamerID();
 }
 
+void UPixelStreaming2Blueprints::GetJsonStringValue(FString Descriptor, FString FieldName, FString& StringValue, bool& Success)
+{
+	UE::PixelStreaming2::ExtractJsonFromDescriptor(Descriptor, FieldName, StringValue, Success);
+}
+
+void UPixelStreaming2Blueprints::AddJsonStringValue(const FString& Descriptor, FString FieldName, FString StringValue, FString& NewDescriptor, bool& Success)
+{
+	UE::PixelStreaming2::ExtendJsonWithField(Descriptor, FieldName, StringValue, NewDescriptor, Success);
+}
+
 UPixelStreaming2Delegates* UPixelStreaming2Blueprints::GetDelegates()
 {
 	return UPixelStreaming2Delegates::Get();
diff --git a/Engine/Plugins/Media/PixelStreaming2/Source/PixelStreaming2/Private/Blueprints/PixelStreaming2Blueprints.h b/Engine/Plugins/Media/PixelStreaming2/Source/PixelStreaming2/Private/Blueprints/PixelStreaming2Blueprints.h
--- a/Engine/Plugins/Media/PixelStreaming2/Source/PixelStreaming2/Private/Blueprints/PixelStreaming2Blueprints.h
+++ b/Engine/Plugins/Media/PixelStreaming2/Source/PixelStreaming2/Private/Blueprints/PixelStreaming2Blueprints.h
@@ -15,6 +15,13 @@ class UPixelStreaming2Blueprints : public UBlueprintFunctionLibrary
 {
 	GENERATED_BODY()
 public:
+	UFUNCTION(BlueprintCallable, Category = "Pixel Streaming 2")
+	static void SendResponse(FString Descriptor);
+
+
+	UFUNCTION(BlueprintCallable, Category = "Pixel Streaming 2")
+	static void StreamerSendResponse(FString StreamerId, FString Descriptor);
+
 	/**
 	 * Send a specified byte array over the WebRTC peer connection data channel. The extension and mime type are used for file reconstruction on the front end
 	 *
@@ -135,6 +142,32 @@ public:
 	UFUNCTION(BlueprintCallable, Category = "Pixel Streaming 2")
 	static FString GetDefaultStreamerID();
 
+	/**
+	 * Helper function to extract a string field from a JSON descriptor of a
+	 * UI interaction given its field name.
+	 * The field name may be hierarchical, delimited by a period. For example,
+	 * to access the Width value of a Resolution command above you should use
+	 * "Resolution.Width" to get the width value.
+	 * @param Descriptor - The UI interaction JSON descriptor.
+	 * @param FieldName - The name of the field to look for in the JSON.
+	 * @param StringValue - The string value associated with the field name.
+	 * @param Success - True if the field exists in the JSON data.
+	 */
+	UFUNCTION(BlueprintPure, Category = "Pixel Streaming 2")
+	static void GetJsonStringValue(FString Descriptor, FString FieldName, FString& StringValue, bool& Success);
+
+	/**
+	 * Helper function to add a string field to a JSON descriptor. This produces
+	 * a new descriptor which may then be chained to add further string fields.
+	 * @param Descriptor - The initial JSON descriptor which may be blank initially.
+	 * @param FieldName - The name of the field to add to the JSON.
+	 * @param StringValue - The string value associated with the field name.
+	 * @param NewDescriptor - The JSON descriptor with the string field added.
+	 * @param Success - True if the string field could be added successfully.
+	 */
+	UFUNCTION(BlueprintPure, Category = "Pixel Streaming 2")
+	static void AddJsonStringValue(const FString& Descriptor, FString FieldName, FString StringValue, FString& NewDescriptor, bool& Success);
+
 	// PixelStreaming2Delegates
 	/**
 	 * Get the singleton. This allows application-specific blueprints to bind
diff --git a/Engine/Plugins/Media/PixelStreaming2/Source/PixelStreaming2/Private/Thread.cpp b/Engine/Plugins/Media/PixelStreaming2/Source/PixelStreaming2/Private/Thread.cpp
--- a/Engine/Plugins/Media/PixelStreaming2/Source/PixelStreaming2/Private/Thread.cpp
+++ b/Engine/Plugins/Media/PixelStreaming2/Source/PixelStreaming2/Private/Thread.cpp
@@ -138,6 +138,11 @@ namespace UE::PixelStreaming2
 
 		void RemoveTask(FPixelStreamingTickableTask* Task)
 		{
+			if (Task == nullptr)
+			{
+				return;
+			}
+			
 			// Lock TaskLock before NewTaskLock to ensure deadlock does not happen when Tick and StartTicking lock.
 			// Locking matches FTickableObjectBase locking in Tickable.cpp
 			FScopeLock TaskLock(&TasksMutex);
diff --git a/Engine/Plugins/Media/PixelStreaming2/Source/PixelStreaming2RTC/Internal/RTCStatsCollector.h b/Engine/Plugins/Media/PixelStreaming2/Source/PixelStreaming2RTC/Internal/RTCStatsCollector.h
new file mode 100644
--- /dev/null
+++ b/Engine/Plugins/Media/PixelStreaming2/Source/PixelStreaming2RTC/Internal/RTCStatsCollector.h
@@ -0,0 +1,242 @@
+// Copyright Epic Games, Inc. All Rights Reserved.
+
+#pragma once
+
+#include "Containers/Map.h"
+#include "Containers/UnrealString.h"
+#include "Templates/SharedPointer.h"
+#include "UObject/NameTypes.h"
+
+#include "epic_rtc/core/stats.h"
+
+namespace UE::PixelStreaming2
+{
+	class FStats;
+	
+	namespace RTCStatCategories
+	{
+		const FString LocalVideoTrack = TEXT("video-track-outbound");
+		const FString LocalAudioTrack = TEXT("audio-track-outbound");
+		const FString VideoSource = TEXT("video-source");
+		const FString AudioSource = TEXT("audio-source");
+		const FString DataChannel = TEXT("data-channel");
+		const FString RemoteVideoTrack = TEXT("video-track-inbound");
+		const FString RemoteAudioTrack = TEXT("audio-track-inbound");
+		const FString CandidatePair = TEXT("candidate-pair");
+	} // namespace RTCStatCategories
+
+	struct PIXELSTREAMING2RTC_API FStatData
+	{
+	public:
+		enum EDisplayFlags : uint8
+		{
+			HIDDEN = 0,
+			TEXT = 1 << 0,
+			GRAPH = 1 << 1,
+		};
+
+	public:
+		FStatData(FName InStatName, double InStatValue, int InNDecimalPlacesToPrint, bool bInSmooth = false)
+			: StatName(InStatName)
+			, StatValue(InStatValue)
+			, NDecimalPlacesToPrint(InNDecimalPlacesToPrint)
+			, bSmooth(bInSmooth)
+		{
+		}
+
+		bool operator==(const FStatData& Other) const
+		{
+			return Equals(Other);
+		}
+
+		bool Equals(const FStatData& Other) const
+		{
+			return StatName == Other.StatName;
+		}
+
+		bool IsHidden() { return DisplayFlags == EDisplayFlags::HIDDEN; }
+		bool ShouldGraph() { return DisplayFlags & EDisplayFlags::GRAPH; }
+		bool ShouldDisplayText() { return DisplayFlags & EDisplayFlags::TEXT; }
+
+		FName			 StatName;
+		double			 StatValue;
+		int				 NDecimalPlacesToPrint;
+		bool			 bSmooth;
+		double			 LastEMA = 0;
+		int				 NumSamples = 0;
+		uint8			 DisplayFlags = EDisplayFlags::TEXT; // Some stats we only wish to store or broadcast, but not display
+		TOptional<FName> Alias;								 // Some stats need an alias that they are stored by/queried by to disambiguate them from other stats
+	};
+
+	FORCEINLINE uint32 GetTypeHash(const FStatData& Obj)
+	{
+		// From UnrealString.h
+		return GetTypeHash(Obj.StatName);
+	}
+
+	class PIXELSTREAMING2RTC_API FRTCStatsCollector
+	{
+	public:
+		static TSharedPtr<FRTCStatsCollector> Create(const FString& PlayerId);
+
+	private:
+		FRTCStatsCollector();
+		explicit FRTCStatsCollector(const FString& PlayerId);
+
+		void OnWebRTCDisableStatsChanged(IConsoleVariable* Var);
+
+	public:
+		void Process(const EpicRtcConnectionStats& InStats);
+
+	private:
+		/**
+		 * ---------- FRTCTrackedStat -------------------
+		 * Tracks a stat from WebRTC. Stores the current value and previous value.
+		 */
+		class FRTCTrackedStat
+		{
+		public:
+			FRTCTrackedStat(FName StatName, int NDecimalPlaces)
+				: LatestStat(StatName, 0.0, NDecimalPlaces)
+			{
+			}
+
+			FRTCTrackedStat(FName StatName, int NDecimalPlaces, uint8 DisplayFlags)
+				: LatestStat(StatName, 0.0, NDecimalPlaces)
+			{
+				LatestStat.DisplayFlags = DisplayFlags;
+			}
+
+			FRTCTrackedStat(FName StatName, FName Alias, int NDecimalPlaces, uint8 DisplayFlags);
+
+			double			 CalculateDelta(double Period) const;
+			double			 Average() const;
+			const FStatData& GetLatestStat() const { return LatestStat; }
+			void			 SetLatestValue(double InValue);
+
+		private:
+			FStatData LatestStat;
+			double	  PrevValue = 0.0;
+		};
+
+		class FStatsSink
+		{
+		public:
+			FStatsSink(FName InCategory);
+			virtual ~FStatsSink() = default;
+			virtual void Add(FName StatName, int NDecimalPlaces)
+			{
+				Stats.Add(StatName, FRTCTrackedStat(StatName, NDecimalPlaces, FStatData::EDisplayFlags::TEXT));
+			}
+
+			virtual void AddAliased(FName StatName, FName AliasedName, int NDecimalPlaces, uint8 DisplayFlags);
+			virtual void AddNonRendered(FName StatName)
+			{
+				Stats.Add(StatName, FRTCTrackedStat(StatName, 2, FStatData::EDisplayFlags::HIDDEN));
+			}
+			virtual void AddStatCalculator(const TFunction<TOptional<FStatData>(FStatsSink&, double)>& Calculator)
+			{
+				Calculators.Add(Calculator);
+			}
+			// returns true if the value is worth storing (false if it started and remains zero)
+			virtual bool			 UpdateValue(double NewValue, FRTCTrackedStat* SetValueHere);
+			virtual FRTCTrackedStat* Get(const FName& StatName)
+			{
+				return Stats.Find(StatName);
+			}
+			virtual FStatData* GetCalculatedStat(const FName& StatName)
+			{
+				return CalculatedStats.Find(StatName);
+			}
+			virtual void PostProcess(FStats* PSStats, const FString& PeerId, double SecondsDelta);
+
+			// Stats that are stored as is.
+			TMap<FName, FRTCTrackedStat> Stats;
+			// Stats we calculate based on the stats map above. This calculation is done in FStatsSink::PostProcess by the `Calculators` below.
+			TMap<FName, FStatData>										 CalculatedStats;
+			TArray<TFunction<TOptional<FStatData>(FStatsSink&, double)>> Calculators;
+
+		protected:
+			FName Category;
+		};
+
+		class FRTPLocalVideoTrackStatsSink : public FStatsSink
+		{
+		public:
+			FRTPLocalVideoTrackStatsSink(FName Category);
+			virtual ~FRTPLocalVideoTrackStatsSink() = default;
+			void Process(const EpicRtcLocalTrackRtpStats& InStats, const FString& PeerId, double SecondsDelta);
+		};
+
+		class FRTPLocalAudioTrackStatsSink : public FStatsSink
+		{
+		public:
+			FRTPLocalAudioTrackStatsSink(FName Category);
+			virtual ~FRTPLocalAudioTrackStatsSink() = default;
+			void Process(const EpicRtcLocalTrackRtpStats& InStats, const FString& PeerId, double SecondsDelta);
+		};
+
+		class FRTPRemoteTrackStatsSink : public FStatsSink
+		{
+		public:
+			FRTPRemoteTrackStatsSink(FName Category);
+			virtual ~FRTPRemoteTrackStatsSink() = default;
+			void Process(const EpicRtcRemoteTrackRtpStats& InStats, const FString& PeerId, double SecondsDelta);
+		};
+
+		class FVideoSourceStatsSink : public FStatsSink
+		{
+		public:
+			FVideoSourceStatsSink(FName Category);
+			virtual ~FVideoSourceStatsSink() = default;
+			void Process(const EpicRtcVideoSourceStats& InStats, const FString& PeerId, double SecondsDelta);
+		};
+
+		class FAudioSourceStatsSink : public FStatsSink
+		{
+		public:
+			FAudioSourceStatsSink(FName Category);
+			virtual ~FAudioSourceStatsSink() = default;
+			void Process(const EpicRtcAudioSourceStats& InStats, const FString& PeerId, double SecondsDelta);
+		};
+
+		/**
+		 * ---------- FDataChannelSink -------------------
+		 */
+		class FDataTrackStatsSink : public FStatsSink
+		{
+		public:
+			FDataTrackStatsSink(FName Category);
+			virtual ~FDataTrackStatsSink() = default;
+
+			void Process(const EpicRtcDataTrackStats& InStats, const FString& PeerId, double SecondsDelta);
+		};
+
+		class FCandidatePairStatsSink : public FStatsSink
+		{
+		public:
+			FCandidatePairStatsSink(FName Category);
+			virtual ~FCandidatePairStatsSink() = default;
+			void Process(const EpicRtcIceCandidatePairStats& InStats, const FString& PeerId, double SecondsDelta);
+		};
+
+	private:
+		FString AssociatedPlayerId;
+		uint64	LastCalculationCycles;
+		bool	bIsEnabled;
+
+		//    index,         ssrc                          sink
+		TMap<uint64, TMap<uint32, TUniquePtr<FRTPLocalVideoTrackStatsSink>>> LocalVideoTrackSinks;
+		TMap<uint64, TMap<uint32, TUniquePtr<FRTPLocalAudioTrackStatsSink>>> LocalAudioTrackSinks;
+
+		TMap<uint64, TMap<uint32, TUniquePtr<FRTPRemoteTrackStatsSink>>> RemoteVideoTrackSinks;
+		TMap<uint64, TMap<uint32, TUniquePtr<FRTPRemoteTrackStatsSink>>> RemoteAudioTrackSinks;
+
+		TMap<uint64, TUniquePtr<FVideoSourceStatsSink>> VideoSourceSinks;
+		TMap<uint64, TUniquePtr<FAudioSourceStatsSink>> AudioSourceSinks;
+
+		TMap<uint64, TUniquePtr<FDataTrackStatsSink>> DataTrackSinks;
+
+		TUniquePtr<FCandidatePairStatsSink> CandidatePairStatsSink;
+	};
+} // namespace UE::PixelStreaming2
\ No newline at end of file
diff --git a/Engine/Plugins/Media/PixelStreaming2/Source/PixelStreaming2RTC/Private/RTCStatsCollector.cpp b/Engine/Plugins/Media/PixelStreaming2/Source/PixelStreaming2RTC/Private/RTCStatsCollector.cpp
--- a/Engine/Plugins/Media/PixelStreaming2/Source/PixelStreaming2RTC/Private/RTCStatsCollector.cpp
+++ b/Engine/Plugins/Media/PixelStreaming2/Source/PixelStreaming2RTC/Private/RTCStatsCollector.cpp
@@ -7,6 +7,7 @@
 #include "PixelStreaming2StatNames.h"
 #include "EpicRtcStreamer.h"
 #include "UtilsString.h"
+#include "Stats.h"
 
 namespace UE::PixelStreaming2
 {
diff --git a/Engine/Plugins/Media/PixelStreaming2/Source/PixelStreaming2RTC/Private/RTCStatsCollector.h b/Engine/Plugins/Media/PixelStreaming2/Source/PixelStreaming2RTC/Private/RTCStatsCollector.h
deleted file mode 100644
diff --git a/Engine/Plugins/Media/PixelStreaming2/Source/PixelStreaming2RTC/Private/Stats.h b/Engine/Plugins/Media/PixelStreaming2/Source/PixelStreaming2RTC/Private/Stats.h
--- a/Engine/Plugins/Media/PixelStreaming2/Source/PixelStreaming2RTC/Private/Stats.h
+++ b/Engine/Plugins/Media/PixelStreaming2/Source/PixelStreaming2RTC/Private/Stats.h
@@ -10,58 +10,10 @@
 #include "Misc/Optional.h"
 #include "PixelCaptureFrameMetadata.h"
 #include "IPixelStreaming2Stats.h"
+#include "RTCStatsCollector.h"
 
 namespace UE::PixelStreaming2
 {
-	struct FStatData
-	{
-	public:
-		enum EDisplayFlags : uint8
-		{
-			HIDDEN = 0,
-			TEXT = 1 << 0,
-			GRAPH = 1 << 1,
-		};
-
-	public:
-		FStatData(FName InStatName, double InStatValue, int InNDecimalPlacesToPrint, bool bInSmooth = false)
-			: StatName(InStatName)
-			, StatValue(InStatValue)
-			, NDecimalPlacesToPrint(InNDecimalPlacesToPrint)
-			, bSmooth(bInSmooth)
-		{
-		}
-
-		bool operator==(const FStatData& Other) const
-		{
-			return Equals(Other);
-		}
-
-		bool Equals(const FStatData& Other) const
-		{
-			return StatName == Other.StatName;
-		}
-
-		bool IsHidden() { return DisplayFlags == EDisplayFlags::HIDDEN; }
-		bool ShouldGraph() { return DisplayFlags & EDisplayFlags::GRAPH; }
-		bool ShouldDisplayText() { return DisplayFlags & EDisplayFlags::TEXT; }
-
-		FName			 StatName;
-		double			 StatValue;
-		int				 NDecimalPlacesToPrint;
-		bool			 bSmooth;
-		double			 LastEMA = 0;
-		int				 NumSamples = 0;
-		uint8			 DisplayFlags = EDisplayFlags::TEXT; // Some stats we only wish to store or broadcast, but not display
-		TOptional<FName> Alias;								 // Some stats need an alias that they are stored by/queried by to disambiguate them from other stats
-	};
-
-	FORCEINLINE uint32 GetTypeHash(const FStatData& Obj)
-	{
-		// From UnrealString.h
-		return GetTypeHash(Obj.StatName);
-	}
-
 	// Stat that can be optionally rendered
 	struct FStoredStat
 	{
