[PS2] Add: Setting to configure the maximum number of attempts the websocket will make to reconnect to the signalling server.

diff --git a/Engine/Plugins/Media/PixelStreaming2/Source/PixelStreaming2/Internal/StreamerReconnectTimer.h b/Engine/Plugins/Media/PixelStreaming2/Source/PixelStreaming2/Internal/StreamerReconnectTimer.h
--- a/Engine/Plugins/Media/PixelStreaming2/Source/PixelStreaming2/Internal/StreamerReconnectTimer.h
+++ b/Engine/Plugins/Media/PixelStreaming2/Source/PixelStreaming2/Internal/StreamerReconnectTimer.h
@@ -14,27 +14,29 @@ namespace UE::PixelStreaming2
 	 * Timer takes a `TWeakPtr<IPixelStreaming2Streamer>`, meaning the streamer object must exist for the timer to do anything.
 	 * Note: this timer runs on the game thread and its reconnect interval is controlled by the CVar
 	 * UPixelStreaming2PluginSettings::CVarSignalingReconnectInterval
-	*/
+	 */
 	class PIXELSTREAMING2_API FStreamerReconnectTimer : public FTickableGameObject
 	{
 	public:
-
 		FStreamerReconnectTimer();
 		virtual ~FStreamerReconnectTimer() = default;
 		void Start(TWeakPtr<IPixelStreaming2Streamer> InWeakStreamer);
 		void Stop();
 
 		// Begin FTickableGameObject
-		virtual void Tick (float DeltaTime) override;
-		virtual bool IsTickableInEditor() const override { return true; }
-		virtual bool IsTickableWhenPaused() const override { return true; }
+		virtual void		Tick(float DeltaTime) override;
+		virtual bool		IsTickableInEditor() const override { return true; }
+		virtual bool		IsTickableWhenPaused() const override { return true; }
 		FORCEINLINE TStatId GetStatId() const { RETURN_QUICK_DECLARE_CYCLE_STAT(PixelStreaming2StreamerReconnectTimer, STATGROUP_Tickables); }
 		// End FTickableGameObject
 
+		DECLARE_EVENT(FStreamerReconnectTimer, FOnExceededMaxReconnectAttempts);
+		FOnExceededMaxReconnectAttempts OnExceededMaxReconnectAttempts;
+
 	private:
 		TWeakPtr<IPixelStreaming2Streamer> WeakStreamer;
-		bool bEnabled = false;
-		uint64 LastReconnectCycles = 0;
-
+		bool							   bEnabled = false;
+		uint64							   NumReconnectAttempts = 0;
+		uint64							   LastReconnectCycles = 0;
 	};
 } // namespace UE::PixelStreaming2
\ No newline at end of file
diff --git a/Engine/Plugins/Media/PixelStreaming2/Source/PixelStreaming2/Private/StreamerReconnectTimer.cpp b/Engine/Plugins/Media/PixelStreaming2/Source/PixelStreaming2/Private/StreamerReconnectTimer.cpp
--- a/Engine/Plugins/Media/PixelStreaming2/Source/PixelStreaming2/Private/StreamerReconnectTimer.cpp
+++ b/Engine/Plugins/Media/PixelStreaming2/Source/PixelStreaming2/Private/StreamerReconnectTimer.cpp
@@ -15,6 +15,7 @@ namespace UE::PixelStreaming2
 	void FStreamerReconnectTimer::Start(TWeakPtr<IPixelStreaming2Streamer> InWeakStreamer)
 	{
 		WeakStreamer = InWeakStreamer;
+		NumReconnectAttempts = 0;
 		bEnabled = true;
 	}
 
@@ -57,15 +58,24 @@ namespace UE::PixelStreaming2
 
 		uint64 CyclesNow = FPlatformTime::Cycles64();
 		uint64 DeltaCycles = CyclesNow - LastReconnectCycles;
-		float DeltaSeconds = FPlatformTime::ToSeconds(DeltaCycles);
+		float  DeltaSeconds = FPlatformTime::ToSeconds(DeltaCycles);
 
 		// If enough time has elapsed, try a reconnect
 		if (DeltaSeconds >= ReconnectInterval)
 		{
-			UE_LOG(LogPixelStreaming2, Log, TEXT("Streamer reconnecting..."))
-				Streamer->StartStreaming();
+			// Check if the next attempt to reconnect will exceed the maximum number of attempts
+			if (UPixelStreaming2PluginSettings::CVarSignalingMaxReconnectAttempts.GetValueOnAnyThread() >= 0 && (NumReconnectAttempts + 1) > UPixelStreaming2PluginSettings::CVarSignalingMaxReconnectAttempts.GetValueOnAnyThread())
+			{
+				// Maxmimum exceeded so don't attempt it and instead stop the timer
+				Stop();
+				OnExceededMaxReconnectAttempts.Broadcast();
+				return;
+			}
+
+			NumReconnectAttempts++;
+			UE_LOGFMT(LogPixelStreaming2, Log, "Streamer reconnecting... Attempt {0}", NumReconnectAttempts);
+			Streamer->StartStreaming();
 			LastReconnectCycles = CyclesNow;
 		}
 	}
-
 } // namespace UE::PixelStreaming2
\ No newline at end of file
diff --git a/Engine/Plugins/Media/PixelStreaming2/Source/PixelStreaming2Input/Internal/InputStructures.h b/Engine/Plugins/Media/PixelStreaming2/Source/PixelStreaming2Input/Internal/InputStructures.h
--- a/Engine/Plugins/Media/PixelStreaming2/Source/PixelStreaming2Input/Internal/InputStructures.h
+++ b/Engine/Plugins/Media/PixelStreaming2/Source/PixelStreaming2Input/Internal/InputStructures.h
@@ -27,8 +27,6 @@ namespace UE::PixelStreaming2Input
 		TPayload(FArchive& Ar)
 		{
 			TPayloadParser<TMakeIntegerSequence<uint32, sizeof...(Types)>>::Parse(*this, Ar);
-
-			checkf(Ar.GetArchiveState().AtEnd(), TEXT("Archive isn't empty after parsing into payload. You may be losing information!"));
 		}
 	};
 } // namespace UE::PixelStreaming2Input
diff --git a/Engine/Plugins/Media/PixelStreaming2/Source/PixelStreaming2RTC/Internal/EpicRtcWebsocket.h b/Engine/Plugins/Media/PixelStreaming2/Source/PixelStreaming2RTC/Internal/EpicRtcWebsocket.h
--- a/Engine/Plugins/Media/PixelStreaming2/Source/PixelStreaming2RTC/Internal/EpicRtcWebsocket.h
+++ b/Engine/Plugins/Media/PixelStreaming2/Source/PixelStreaming2RTC/Internal/EpicRtcWebsocket.h
@@ -15,7 +15,7 @@ namespace UE::PixelStreaming2
 	class PIXELSTREAMING2RTC_API FEpicRtcWebsocket : public EpicRtcWebsocketInterface, public FTickableGameObject
 	{
 	public:
-		FEpicRtcWebsocket(bool bKeepAlive = true, TSharedPtr<IWebSocket> WebSocket = nullptr);
+		FEpicRtcWebsocket(bool bKeepAlive = true, TSharedPtr<IWebSocket> WebSocket = nullptr, TFunction<void(void)> OnMaxReconnectAttemptsExceeded = []() {});
 		virtual ~FEpicRtcWebsocket() = default;
 
 		// Begin EpicRtcWebsocketInterface
@@ -49,13 +49,16 @@ namespace UE::PixelStreaming2
 
 	private:
 		FString											Url;
+		FString											LastError;
 		TSharedPtr<IWebSocket>							WebSocket;
 		TRefCountPtr<EpicRtcWebsocketObserverInterface> Observer;
 		uint64											LastKeepAliveCycles = 0;
 		uint64											LastReconnectCycles = 0;
+		uint64											NumReconnectAttempts = 0;
 		bool											bSendKeepAlive = false;
 		bool											bReconnectOnError = false;
 		bool											bCloseRequested = false;
+		TFunction<void(void)>							OnMaxReconnectAttemptsExceeded;
 
 	public:
 		// Begin EpicRtcRefCountInterface
diff --git a/Engine/Plugins/Media/PixelStreaming2/Source/PixelStreaming2RTC/Internal/EpicRtcWebsocketFactory.h b/Engine/Plugins/Media/PixelStreaming2/Source/PixelStreaming2RTC/Internal/EpicRtcWebsocketFactory.h
--- a/Engine/Plugins/Media/PixelStreaming2/Source/PixelStreaming2RTC/Internal/EpicRtcWebsocketFactory.h
+++ b/Engine/Plugins/Media/PixelStreaming2/Source/PixelStreaming2RTC/Internal/EpicRtcWebsocketFactory.h
@@ -13,8 +13,9 @@ namespace UE::PixelStreaming2
 	class PIXELSTREAMING2RTC_API FEpicRtcWebsocketFactory : public EpicRtcWebsocketFactoryInterface
 	{
 	public:
-		FEpicRtcWebsocketFactory(bool bInSendKeepAlive = true)
+		FEpicRtcWebsocketFactory(bool bInSendKeepAlive = true, TFunction<void(void)> OnMaxReconnectAttemptsExceeded = []() {})
 			: bSendKeepAlive(bInSendKeepAlive)
+			, OnMaxReconnectAttemptsExceeded(OnMaxReconnectAttemptsExceeded)
 		{
 		}
 		virtual ~FEpicRtcWebsocketFactory() = default;
@@ -28,6 +29,8 @@ namespace UE::PixelStreaming2
 
 	private:
 		bool bSendKeepAlive;
+
+		TFunction<void(void)> OnMaxReconnectAttemptsExceeded;
 	};
 
 } // namespace UE::PixelStreaming2
diff --git a/Engine/Plugins/Media/PixelStreaming2/Source/PixelStreaming2RTC/Private/EpicRtcWebsocket.cpp b/Engine/Plugins/Media/PixelStreaming2/Source/PixelStreaming2RTC/Private/EpicRtcWebsocket.cpp
--- a/Engine/Plugins/Media/PixelStreaming2/Source/PixelStreaming2RTC/Private/EpicRtcWebsocket.cpp
+++ b/Engine/Plugins/Media/PixelStreaming2/Source/PixelStreaming2RTC/Private/EpicRtcWebsocket.cpp
@@ -16,9 +16,10 @@ namespace UE::PixelStreaming2
 	DECLARE_LOG_CATEGORY_EXTERN(LogEpicRtcWebsocket, Log, All);
 	DEFINE_LOG_CATEGORY(LogEpicRtcWebsocket);
 
-	FEpicRtcWebsocket::FEpicRtcWebsocket(bool bKeepAlive, TSharedPtr<IWebSocket> InWebSocket)
+	FEpicRtcWebsocket::FEpicRtcWebsocket(bool bKeepAlive, TSharedPtr<IWebSocket> InWebSocket, TFunction<void(void)> InOnMaxReconnectAttemptsExceeded)
 		: WebSocket(InWebSocket)
 		, bSendKeepAlive(bKeepAlive)
+		, OnMaxReconnectAttemptsExceeded(InOnMaxReconnectAttemptsExceeded)
 	{
 	}
 
@@ -121,6 +122,7 @@ namespace UE::PixelStreaming2
 		UE_LOG(LogEpicRtcWebsocket, Log, TEXT("Websocket connection made to: %s"), *Url);
 		bCloseRequested = false;
 		bReconnectOnError = false;
+		NumReconnectAttempts = 0;
 		LastKeepAliveCycles = FPlatformTime::Cycles64();
 		Observer->OnOpen();
 	}
@@ -132,7 +134,7 @@ namespace UE::PixelStreaming2
 			// Reconnecting case where we had not connected yet and got an error while connecting (e.g. server not up)
 			bReconnectOnError = true;
 
-			UE_LOG(LogEpicRtcWebsocket, Log, TEXT("Failed to connect to %s. (\"%s\") - signalling server may not be up yet. Reconnecting..."), *Url, *Error);
+			LastError = Error;
 
 			// To reconnect we must close the existing WS (amusingly this does no trigger WS `OnClosed`)
 			WebSocket->Close();
@@ -281,6 +283,18 @@ namespace UE::PixelStreaming2
 		// If enough time has elapsed, try a reconnect
 		if (DeltaSeconds >= ReconnectIntervalSeconds)
 		{
+			// Check if the next attempt to reconnect will exceed the maximum number of attempts
+			if (UPixelStreaming2PluginSettings::CVarSignalingMaxReconnectAttempts.GetValueOnAnyThread() >= 0 && (NumReconnectAttempts + 1) > UPixelStreaming2PluginSettings::CVarSignalingMaxReconnectAttempts.GetValueOnAnyThread())
+			{
+				// Maxmimum exceeded so don't attempt it and instead stop the timer
+				UE_LOGFMT(LogEpicRtcWebsocket, Warning, "Maximum number of reconnect attempts ({0}) exceeded!", UPixelStreaming2PluginSettings::CVarSignalingMaxReconnectAttempts.GetValueOnAnyThread());
+				bReconnectOnError = false;
+				OnMaxReconnectAttemptsExceeded();
+				return;
+			}
+
+			NumReconnectAttempts++;
+			UE_LOGFMT(LogEpicRtcWebsocket, Log, "Failed to connect to {0}. (\"{1}\") - signalling server may not be up yet. Reconnecting... Attempt: {2}", Url, LastError, NumReconnectAttempts);
 			WebSocket->Connect();
 			LastReconnectCycles = CyclesNow;
 			// Do not try to reconnect until we hear the next error
diff --git a/Engine/Plugins/Media/PixelStreaming2/Source/PixelStreaming2RTC/Private/EpicRtcWebsocketFactory.cpp b/Engine/Plugins/Media/PixelStreaming2/Source/PixelStreaming2RTC/Private/EpicRtcWebsocketFactory.cpp
--- a/Engine/Plugins/Media/PixelStreaming2/Source/PixelStreaming2RTC/Private/EpicRtcWebsocketFactory.cpp
+++ b/Engine/Plugins/Media/PixelStreaming2/Source/PixelStreaming2RTC/Private/EpicRtcWebsocketFactory.cpp
@@ -7,7 +7,7 @@ namespace UE::PixelStreaming2
 
 	EpicRtcErrorCode FEpicRtcWebsocketFactory::CreateWebsocket(EpicRtcWebsocketInterface** OutWebsocket)
 	{
-		EpicRtcWebsocketInterface* Websocket = new FEpicRtcWebsocket(bSendKeepAlive);
+		EpicRtcWebsocketInterface* Websocket = new FEpicRtcWebsocket(bSendKeepAlive, nullptr, OnMaxReconnectAttemptsExceeded);
 		// Because the ptr was created with new, we need to call AddRef ourself (ms spec compliant)
 		Websocket->AddRef();
 
diff --git a/Engine/Plugins/Media/PixelStreaming2/Source/PixelStreaming2Settings/Internal/PixelStreaming2PluginSettings.h b/Engine/Plugins/Media/PixelStreaming2/Source/PixelStreaming2Settings/Internal/PixelStreaming2PluginSettings.h
--- a/Engine/Plugins/Media/PixelStreaming2/Source/PixelStreaming2Settings/Internal/PixelStreaming2PluginSettings.h
+++ b/Engine/Plugins/Media/PixelStreaming2/Source/PixelStreaming2Settings/Internal/PixelStreaming2PluginSettings.h
@@ -112,10 +112,17 @@ public:
 	static TAutoConsoleVariable<float> CVarSignalingReconnectInterval;
 	UPROPERTY(config, EditAnywhere, Category = "PixelStreaming", meta = (
 		DisplayName = "Signaling Reconnect Interval",
-		ToolTip = "Frame rate factor to wait for a captured frame when streaming in decoupled mode. Higher factor waits longer but may also result in higher latency."
+		ToolTip = "Changes the number of seconds between attempted reconnects to the signaling server. This is useful for reducing the log spam produced from attempted reconnects. A value <= 0 results in no reconnect. Default: 2.0s"
 		))
 	float SignalingReconnectInterval = 2.0f;
 
+	static TAutoConsoleVariable<float> CVarSignalingMaxReconnectAttempts;
+	UPROPERTY(config, EditAnywhere, Category = "PixelStreaming", meta = (
+		DisplayName = "Signaling Max Reconnect Attempts",
+		ToolTip = "Changes the number of attempts that will be made to reconnect to the signalling server. This is useful for triggering application shutdowns if this value is exceeded. A value of < 0 results in unlimited attempts. Default: -1"
+		))
+	int32 SignalingMaxReconnectAttempts = -1;
+
 	static TAutoConsoleVariable<float> CVarSignalingKeepAliveInterval;
 	UPROPERTY(config, EditAnywhere, Category = "PixelStreaming", meta = (
 		DisplayName = "Signaling Reconnect Interval",
diff --git a/Engine/Plugins/Media/PixelStreaming2/Source/PixelStreaming2Settings/Private/PixelStreaming2PluginSettings.cpp b/Engine/Plugins/Media/PixelStreaming2/Source/PixelStreaming2Settings/Private/PixelStreaming2PluginSettings.cpp
--- a/Engine/Plugins/Media/PixelStreaming2/Source/PixelStreaming2Settings/Private/PixelStreaming2PluginSettings.cpp
+++ b/Engine/Plugins/Media/PixelStreaming2/Source/PixelStreaming2Settings/Private/PixelStreaming2PluginSettings.cpp
@@ -63,25 +63,25 @@ namespace
 	static void VerifyCVarDefaultStreamerType(IConsoleVariable* CVar)
 	{
 		TArray<FString> AvailableFactoryTypes = IPixelStreaming2StreamerFactory::GetAvailableFactoryTypes();
-		FString SpecifiedFactory = CVar->GetString();
+		FString			SpecifiedFactory = CVar->GetString();
 
-		if(AvailableFactoryTypes.Num() == 0)
+		if (AvailableFactoryTypes.Num() == 0)
 		{
 			// This code path executes when the cvar is initially set and no factories have been registered
 			return;
 		}
 
 		bool bValid = false;
-		for(const FString& AvailableFactory : AvailableFactoryTypes)
+		for (const FString& AvailableFactory : AvailableFactoryTypes)
 		{
-			if(SpecifiedFactory == AvailableFactory)
+			if (SpecifiedFactory == AvailableFactory)
 			{
-				bValid =  true;
+				bValid = true;
 				break;
 			}
 		}
 
-		if(!bValid)
+		if (!bValid)
 		{
 			UE_LOGFMT(LogPixelStreaming2Settings, Warning, "\"{0}\" isn't a registered streamer type. Valid types: [{1}]. Restoring to \"{2}\"", SpecifiedFactory, FString::Join(AvailableFactoryTypes, TEXT(",")), CVar->GetDefaultValue());
 			CVar->SetWithCurrentPriority(*CVar->GetDefaultValue());
@@ -281,6 +281,12 @@ TAutoConsoleVariable<float> UPixelStreaming2PluginSettings::CVarSignalingReconne
 	TEXT("Changes the number of seconds between attempted reconnects to the signaling server. This is useful for reducing the log spam produced from attempted reconnects. A value <= 0 results in no reconnect. Default: 2.0s"),
 	ECVF_Default);
 
+TAutoConsoleVariable<float> UPixelStreaming2PluginSettings::CVarSignalingMaxReconnectAttempts(
+	TEXT("PixelStreaming2.SignalingMaxReconnectAttempts"),
+	-1,
+	TEXT("Changes the number of attempts that will be made to reconnect to the signalling server. This is useful for triggering application shutdowns if this value is exceeded. A value of < 0 results in unlimited attempts. Default: -1"),
+	ECVF_Default);
+
 TAutoConsoleVariable<float> UPixelStreaming2PluginSettings::CVarSignalingKeepAliveInterval(
 	TEXT("PixelStreaming2.SignalingKeepAliveInterval"),
 	30.0f,
