[PS2] Add: Option to configure whether the default stream starts automatically after module load. Default: true

diff --git a/Engine/Plugins/Media/PixelStreaming2/Source/PixelStreaming2/Private/PixelStreaming2Module.cpp b/Engine/Plugins/Media/PixelStreaming2/Source/PixelStreaming2/Private/PixelStreaming2Module.cpp
--- a/Engine/Plugins/Media/PixelStreaming2/Source/PixelStreaming2/Private/PixelStreaming2Module.cpp
+++ b/Engine/Plugins/Media/PixelStreaming2/Source/PixelStreaming2/Private/PixelStreaming2Module.cpp
@@ -128,7 +128,10 @@ namespace UE::PixelStreaming2
 				// We don't want to create the default streamer when using the editor
 				InitDefaultStreamer();
 
-				DefaultStreamer->StartStreaming();
+				if (UPixelStreaming2PluginSettings::CVarAutoStartStream.GetValueOnAnyThread())
+				{
+					DefaultStreamer->StartStreaming();
+				}
 			}
 		});
 	}
@@ -197,7 +200,7 @@ namespace UE::PixelStreaming2
 		TSharedPtr<IPixelStreaming2Streamer> NewStreamer = Factory->CreateNewStreamer(StreamerId);
 		NewStreamer->Initialize();
 		// Any time we create a new streamer, populate it's signalling server URL with whatever is in the ini, console or command line
-		if(!UPixelStreaming2PluginSettings::CVarConnectionURL.GetValueOnAnyThread().IsEmpty())
+		if (!UPixelStreaming2PluginSettings::CVarConnectionURL.GetValueOnAnyThread().IsEmpty())
 		{
 			NewStreamer->SetConnectionURL(UPixelStreaming2PluginSettings::CVarConnectionURL.GetValueOnAnyThread());
 		}
diff --git a/Engine/Plugins/Media/PixelStreaming2/Source/PixelStreaming2Settings/Internal/PixelStreaming2PluginSettings.h b/Engine/Plugins/Media/PixelStreaming2/Source/PixelStreaming2Settings/Internal/PixelStreaming2PluginSettings.h
--- a/Engine/Plugins/Media/PixelStreaming2/Source/PixelStreaming2Settings/Internal/PixelStreaming2PluginSettings.h
+++ b/Engine/Plugins/Media/PixelStreaming2/Source/PixelStreaming2Settings/Internal/PixelStreaming2PluginSettings.h
@@ -152,6 +152,13 @@ public:
 		))
 	FString DefaultStreamerType = TEXT("DefaultRtc");
 
+	static TAutoConsoleVariable<bool> CVarAutoStartStream;
+	UPROPERTY(config, EditAnywhere, Category = "PixelStreaming", meta = (
+		DisplayName = "Automatically Start Streaming",
+		ToolTip = "Configure the PixelStreaming2 plugin to automatically start streaming once loaded (if not in editor). You may wish to set this value to false and manually call StartStreaming at a later point from your c++ code. Default: true"
+		))
+	bool AutoStartStream = true;
+
 	static TAutoConsoleVariable<FString> CVarConnectionURL;
 	UPROPERTY(config, EditAnywhere, Category = "PixelStreaming", meta = (
 		DisplayName = "Default Connection URL",
diff --git a/Engine/Plugins/Media/PixelStreaming2/Source/PixelStreaming2Settings/Private/PixelStreaming2PluginSettings.cpp b/Engine/Plugins/Media/PixelStreaming2/Source/PixelStreaming2Settings/Private/PixelStreaming2PluginSettings.cpp
--- a/Engine/Plugins/Media/PixelStreaming2/Source/PixelStreaming2Settings/Private/PixelStreaming2PluginSettings.cpp
+++ b/Engine/Plugins/Media/PixelStreaming2/Source/PixelStreaming2Settings/Private/PixelStreaming2PluginSettings.cpp
@@ -167,10 +167,12 @@ static const TSet<TPair<FString, FString>> GetCmdArg = {
 	{ "PixelStreaming2.DecoupleFrameRate", "DecoupleFramerate" },
 	{ "PixelStreaming2.DecoupleWaitFactor", "DecoupleWaitFactor" },
 	{ "PixelStreaming2.SignalingReconnectInterval", "SignalingReconnectInterval" },
+	{ "PixelStreaming2.SignalingMaxReconnectAttempts", "SignalingMaxReconnectAttempts" },
 	{ "PixelStreaming2.SignalingKeepAliveInterval", "SignalingKeepAliveInterval" },
 	{ "PixelStreaming2.UseMediaCapture", "UseMediaCapture" },
 	{ "PixelStreaming2.ID", "DefaultStreamerID" },
 	{ "PixelStreaming2.DefaultStreamerType", "DefaultStreamerType" },
+	{ "PixelStreaming2.AutoStartStream", "AutoStartStream" },
 	{ "PixelStreaming2.ConnectionURL", "ConnectionURL" },
 	{ "PixelStreaming2.CaptureUseFence", "CaptureUseFence" },
 	{ "PixelStreaming2.Encoder.Codec", "Codec" },
@@ -312,6 +314,12 @@ TAutoConsoleVariable<FString> UPixelStreaming2PluginSettings::CVarDefaultStreame
 	FConsoleVariableDelegate::CreateLambda([](IConsoleVariable* Var) { VerifyCVarDefaultStreamerType(Var); }),
 	ECVF_Default);
 
+TAutoConsoleVariable<bool> UPixelStreaming2PluginSettings::CVarAutoStartStream(
+	TEXT("PixelStreaming2.AutoStartStream"),
+	true,
+	TEXT("Configure the PixelStreaming2 plugin to automatically start streaming once loaded (if not in editor). You may wish to set this value to false and manually call StartStreaming at a later point from your c++ code. Default: true"),
+	ECVF_Default);
+
 TAutoConsoleVariable<FString> UPixelStreaming2PluginSettings::CVarConnectionURL(
 	TEXT("PixelStreaming2.ConnectionURL"),
 	TEXT(""),
