[PS2] QOL: Updates to stats. Stopped tracking encode and packetization time as application stats as we are no longer stream sharing so these are per peer. Added logic to remove old video capturer stats when a new capturer is created

diff --git a/Engine/Plugins/Media/PixelStreaming2/Source/PixelStreaming2/Internal/VideoCapturer.h b/Engine/Plugins/Media/PixelStreaming2/Source/PixelStreaming2/Internal/VideoCapturer.h
--- a/Engine/Plugins/Media/PixelStreaming2/Source/PixelStreaming2/Internal/VideoCapturer.h
+++ b/Engine/Plugins/Media/PixelStreaming2/Source/PixelStreaming2/Internal/VideoCapturer.h
@@ -42,6 +42,13 @@ namespace UE::PixelStreaming2
 		DECLARE_MULTICAST_DELEGATE(FOnFrameCaptured);
 		FOnFrameCaptured OnFrameCaptured;
 
+		/**
+		 * Each time the frame capturer is created. This could be triggered by a multitude of situations such as cvar changes, frame 
+		 * resizes, etc
+		 */
+		DECLARE_MULTICAST_DELEGATE(FOnFrameCapturerCreated);
+		FOnFrameCapturerCreated OnFrameCapturerCreated;
+
 	protected:
 		FVideoCapturer(TSharedPtr<FVideoProducer> VideoProducer);
 
diff --git a/Engine/Plugins/Media/PixelStreaming2/Source/PixelStreaming2/Private/VideoCapturer.cpp b/Engine/Plugins/Media/PixelStreaming2/Source/PixelStreaming2/Private/VideoCapturer.cpp
--- a/Engine/Plugins/Media/PixelStreaming2/Source/PixelStreaming2/Private/VideoCapturer.cpp
+++ b/Engine/Plugins/Media/PixelStreaming2/Source/PixelStreaming2/Private/VideoCapturer.cpp
@@ -140,6 +140,8 @@ namespace UE::PixelStreaming2
 
 	TSharedPtr<FPixelCaptureCapturer> FVideoCapturer::CreateCapturer(int32 FinalFormat, float FinalScale)
 	{
+		OnFrameCapturerCreated.Broadcast();
+		
 		if (LastFrameType == PixelCaptureBufferFormat::FORMAT_RHI)
 		{
 			switch (FinalFormat)
diff --git a/Engine/Plugins/Media/PixelStreaming2/Source/PixelStreaming2RTC/Private/EpicRtcStreamer.cpp b/Engine/Plugins/Media/PixelStreaming2/Source/PixelStreaming2RTC/Private/EpicRtcStreamer.cpp
--- a/Engine/Plugins/Media/PixelStreaming2/Source/PixelStreaming2RTC/Private/EpicRtcStreamer.cpp
+++ b/Engine/Plugins/Media/PixelStreaming2/Source/PixelStreaming2RTC/Private/EpicRtcStreamer.cpp
@@ -188,6 +188,8 @@ namespace UE::PixelStreaming2
 		Streamer->EpicRtcManager->OnDataTrackState.AddSP(Streamer.ToSharedRef(), &FEpicRtcStreamer::OnDataTrackState);
 		Streamer->EpicRtcManager->OnDataTrackMessage.AddSP(Streamer.ToSharedRef(), &FEpicRtcStreamer::OnDataTrackMessage);
 
+		VideoCapturer->OnFrameCapturerCreated.AddSP(Streamer.ToSharedRef(), &FEpicRtcStreamer::OnFrameCapturerCreated);
+
 		FPixelStreaming2RTCModule::GetModule()->GetStatsCollector()->OnStatsReady.AddSP(Streamer.ToSharedRef(), &FEpicRtcStreamer::OnStatsReady);
 	}
 
@@ -365,6 +367,23 @@ namespace UE::PixelStreaming2
 		PlayerContext->StatsCollector->Process(ConnectionStats);
 	}
 
+	void FEpicRtcStreamer::OnFrameCapturerCreated()
+	{
+		if (FStats* PSStats = FStats::Get())
+		{
+			PSStats->RemoveAllApplicationStats();
+
+			// Re-add the cvar based application stats so that they're at the top
+			PSStats->StoreApplicationStat(FStatData(FName(TEXT("PixelStreaming2.Encoder.MinQuality")), UPixelStreaming2PluginSettings::CVarEncoderMinQuality.GetValueOnAnyThread(), 0));
+			PSStats->StoreApplicationStat(FStatData(FName(TEXT("PixelStreaming2.Encoder.MaxQuality")), UPixelStreaming2PluginSettings::CVarEncoderMaxQuality.GetValueOnAnyThread(), 0));
+			PSStats->StoreApplicationStat(FStatData(FName(TEXT("PixelStreaming2.Encoder.KeyframeInterval (frames)")), UPixelStreaming2PluginSettings::CVarEncoderKeyframeInterval.GetValueOnAnyThread(), 0));
+			PSStats->StoreApplicationStat(FStatData(FName(TEXT("PixelStreaming2.WebRTC.Fps")), UPixelStreaming2PluginSettings::CVarWebRTCFps.GetValueOnAnyThread(), 0));
+			PSStats->StoreApplicationStat(FStatData(FName(TEXT("PixelStreaming2.WebRTC.StartBitrate")), UPixelStreaming2PluginSettings::CVarWebRTCStartBitrate.GetValueOnAnyThread(), 0));
+			PSStats->StoreApplicationStat(FStatData(FName(TEXT("PixelStreaming2.WebRTC.MinBitrate")), UPixelStreaming2PluginSettings::CVarWebRTCMinBitrate.GetValueOnAnyThread(), 0));
+			PSStats->StoreApplicationStat(FStatData(FName(TEXT("PixelStreaming2.WebRTC.MaxBitrate")), UPixelStreaming2PluginSettings::CVarWebRTCMaxBitrate.GetValueOnAnyThread(), 0));
+		}
+	}
+
 	void FEpicRtcStreamer::OnUIInteraction(FMemoryReader Ar)
 	{
 		FString Res;
diff --git a/Engine/Plugins/Media/PixelStreaming2/Source/PixelStreaming2RTC/Private/EpicRtcStreamer.h b/Engine/Plugins/Media/PixelStreaming2/Source/PixelStreaming2RTC/Private/EpicRtcStreamer.h
--- a/Engine/Plugins/Media/PixelStreaming2/Source/PixelStreaming2RTC/Private/EpicRtcStreamer.h
+++ b/Engine/Plugins/Media/PixelStreaming2/Source/PixelStreaming2RTC/Private/EpicRtcStreamer.h
@@ -97,6 +97,7 @@ namespace UE::PixelStreaming2
 		void RemoveSession(bool bDisconnect);
 		void RemoveRoom();
 		void OnStatsReady(const FString& PlayerId, const EpicRtcConnectionStats& ConnectionStats);
+		void OnFrameCapturerCreated();
 		void OnUIInteraction(FMemoryReader Ar);
 		void OnSendMessage(FString MessageName, FMemoryReader Ar);
 
diff --git a/Engine/Plugins/Media/PixelStreaming2/Source/PixelStreaming2RTC/Private/Stats.cpp b/Engine/Plugins/Media/PixelStreaming2/Source/PixelStreaming2RTC/Private/Stats.cpp
--- a/Engine/Plugins/Media/PixelStreaming2/Source/PixelStreaming2RTC/Private/Stats.cpp
+++ b/Engine/Plugins/Media/PixelStreaming2/Source/PixelStreaming2RTC/Private/Stats.cpp
@@ -86,7 +86,7 @@ namespace UE::PixelStreaming2
 		if (const FPeerStats* SinglePeerStats = PeerStats.Find(PlayerId))
 		{
 			const TMap<FName, FStatGroup>& StatGroups = SinglePeerStats->GetStatGroups();
-			TArray<FName>			 StatCategories;
+			TArray<FName>				   StatCategories;
 			StatGroups.GetKeys(StatCategories);
 
 			// Stat groups contain a name as well as additional info like track index and ssrc
@@ -220,6 +220,12 @@ namespace UE::PixelStreaming2
 		}
 	}
 
+	void FStats::RemoveAllApplicationStats()
+	{
+		FScopeLock Lock(&ApplicationStatsCS);
+		ApplicationStats.Empty();
+	}
+
 	void FStats::FireStatChanged(const FString& PlayerId, FName StatName, float StatValue)
 	{
 		// Broadcast must be done on the GameThread because the GameThread can remove the delegates.
@@ -647,8 +653,6 @@ namespace UE::PixelStreaming2
 		const double TimeCPU = AddTimeStat(FrameMetadata.CaptureProcessCPUTime, FString::Printf(TEXT("%s Layer %d Frame Capture CPU Time"), *FrameMetadata.ProcessName, FrameMetadata.Layer));
 		const double TimeGPUDelay = AddTimeStat(FrameMetadata.CaptureProcessGPUDelay, FString::Printf(TEXT("%s Layer %d Frame Capture GPU Delay Time"), *FrameMetadata.ProcessName, FrameMetadata.Layer));
 		const double TimeGPU = AddTimeStat(FrameMetadata.CaptureProcessGPUTime, FString::Printf(TEXT("%s Layer %d Frame Capture GPU Time"), *FrameMetadata.ProcessName, FrameMetadata.Layer));
-		const double TimeEncode = AddTimeDeltaStat(FrameMetadata.LastEncodeEndTime, FrameMetadata.LastEncodeStartTime, FString::Printf(TEXT("%s Layer %d Frame Encode Time"), *FrameMetadata.ProcessName, FrameMetadata.Layer));
-		const double TimePacketize = AddTimeDeltaStat(FrameMetadata.LastPacketizationEndTime, FrameMetadata.LastPacketizationStartTime, FString::Printf(TEXT("%s Layer %d Frame Packetization Time"), *FrameMetadata.ProcessName, FrameMetadata.Layer));
 
 		const FStatData UseData{ FName(*FString::Printf(TEXT("%s Layer %d Frame Uses"), *FrameMetadata.ProcessName, FrameMetadata.Layer)), static_cast<double>(FrameMetadata.UseCount), 0, false };
 		StoreApplicationStat(UseData);
@@ -658,8 +662,6 @@ namespace UE::PixelStreaming2
 		GraphValue(*FString::Printf(TEXT("%d CPU Time"), FrameMetadata.Layer), StaticCast<float>(TimeCPU), Samples, 0.0f, 30.0f);
 		GraphValue(*FString::Printf(TEXT("%d GPU Delay Time"), FrameMetadata.Layer), StaticCast<float>(TimeGPUDelay), Samples, 0.0f, 30.0f);
 		GraphValue(*FString::Printf(TEXT("%d GPU Time"), FrameMetadata.Layer), StaticCast<float>(TimeGPU), Samples, 0.0f, 30.0f);
-		GraphValue(*FString::Printf(TEXT("%d Encode Time"), FrameMetadata.Layer), StaticCast<float>(TimeEncode), Samples, 0.0f, 10.0f);
-		GraphValue(*FString::Printf(TEXT("%d Packetization Time"), FrameMetadata.Layer), StaticCast<float>(TimePacketize), Samples, 0.0f, 10.0f);
 		GraphValue(*FString::Printf(TEXT("%d Frame Uses"), FrameMetadata.Layer), StaticCast<float>(FrameMetadata.UseCount), Samples, 0.0f, 10.0f);
 	}
 
diff --git a/Engine/Plugins/Media/PixelStreaming2/Source/PixelStreaming2RTC/Private/Stats.h b/Engine/Plugins/Media/PixelStreaming2/Source/PixelStreaming2RTC/Private/Stats.h
--- a/Engine/Plugins/Media/PixelStreaming2/Source/PixelStreaming2RTC/Private/Stats.h
+++ b/Engine/Plugins/Media/PixelStreaming2/Source/PixelStreaming2RTC/Private/Stats.h
@@ -94,6 +94,7 @@ namespace UE::PixelStreaming2
 		void RemoveAllPeerStats();
 		void StorePeerStat(const FString& PlayerId, FName StatCategory, FStatData Stat);
 		void StoreApplicationStat(FStatData PeerStat);
+		void RemoveAllApplicationStats();
 		void Tick(float DeltaTime);
 
 		bool  OnToggleStats(UWorld* World, FCommonViewportClient* ViewportClient, const TCHAR* Stream);
