[PS2] Fix: Possible incorrect message construction due to not forwarding variables.

diff --git a/Engine/Plugins/Media/PixelStreaming2/Source/PixelStreaming2RTC/Internal/EpicRtcDataTrack.h b/Engine/Plugins/Media/PixelStreaming2/Source/PixelStreaming2RTC/Internal/EpicRtcDataTrack.h
--- a/Engine/Plugins/Media/PixelStreaming2/Source/PixelStreaming2RTC/Internal/EpicRtcDataTrack.h
+++ b/Engine/Plugins/Media/PixelStreaming2/Source/PixelStreaming2RTC/Internal/EpicRtcDataTrack.h
@@ -10,36 +10,54 @@
 namespace UE::PixelStreaming2
 {
 	template <typename T>
-	inline size_t ValueSize(T&& Value)
+	inline size_t GetByteSizeOf(const T&& Value)
 	{
+		// Requires being r-value reference because the string and array specialization
+		// is not always detected and forwarding allows the right specialization to be used.
+		static_assert(std::is_rvalue_reference<decltype(Value)>::value);
 		return sizeof(Value);
 	}
 
 	template <typename T>
-	inline const void* ValueLoc(T&& Value)
+	inline const void* GetDataPointer(const T&& Value)
 	{
+		// Requires being r-value reference because the string and array specialization
+		// is not always detected and forwarding allows the right specialization to be used.
+		static_assert(std::is_rvalue_reference<decltype(Value)>::value);
 		return &Value;
 	}
 
-	inline size_t ValueSize(FString&& Value)
+	inline size_t GetByteSizeOf(const FString&& Value)
 	{
+		// Requires being r-value reference because the string and array specialization
+		// is not always detected and forwarding allows the right specialization to be used.
+		static_assert(std::is_rvalue_reference<decltype(Value)>::value);
 		return Value.Len() * sizeof(TCHAR);
 	}
 
-	inline const void* ValueLoc(FString&& Value)
+	inline const void* GetDataPointer(const FString&& Value)
 	{
+		// Requires being r-value reference because the string and array specialization
+		// is not always detected and forwarding allows the right specialization to be used.
+		static_assert(std::is_rvalue_reference<decltype(Value)>::value);
 		return *Value;
 	}
 
 	template <typename T>
-	inline size_t ValueSize(TArray<T>&& Value)
+	inline size_t GetByteSizeOf(const TArray<T>&& Value)
 	{
+		// Requires being r-value reference because the string and array specialization
+		// is not always detected and forwarding allows the right specialization to be used.
+		static_assert(std::is_rvalue_reference<decltype(Value)>::value);
 		return Value.Num() * sizeof(T);
 	}
 
 	template <typename T>
-	inline const void* ValueLoc(TArray<T>&& Value)
+	inline const void* GetDataPointer(const TArray<T>&& Value)
 	{
+		// Requires being r-value reference because the string and array specialization
+		// is not always detected and forwarding allows the right specialization to be used.
+		static_assert(std::is_rvalue_reference<decltype(Value)>::value);
 		return Value.GetData();
 	}
 
@@ -62,10 +80,14 @@ namespace UE::PixelStreaming2
 		}
 
 		template <typename T>
-		void Insert(T&& Value)
+		void Insert(const T&& Value)
 		{
-			const size_t VSize = ValueSize(Forward<T>(Value));
-			const void*	 VLoc = ValueLoc(Forward<T>(Value));
+			// Requires being r-value reference because the string and array specialization
+			// is not always detected and forwarding allows the right specialization to be used.
+			static_assert(std::is_rvalue_reference<decltype(Value)>::value);
+
+			const size_t VSize = GetByteSizeOf(Forward<const T>(Value));
+			const void*	 VLoc = GetDataPointer(Forward<const T>(Value));
 			check(Pos + VSize <= Buffer.Num());
 			FMemory::Memcpy(Buffer.GetData() + Pos, VLoc, VSize);
 			Pos += VSize;
@@ -149,9 +171,9 @@ namespace UE::PixelStreaming2
 		FBufferBuilder EncodeMessage(uint8 Type, Args... VarArgs) const
 		{
 
-			FBufferBuilder Builder(sizeof(Type) + (0 + ... + ValueSize(Forward<Args>(VarArgs))));
+			FBufferBuilder Builder(sizeof(Type) + (0 + ... + GetByteSizeOf(Forward<Args>(VarArgs))));
 			PrependData(Builder);
-			Builder.Insert(Type);
+			Builder.Insert(Forward<uint8>(Type));
 			(Builder.Insert(Forward<Args>(VarArgs)), ...);
 
 			return MoveTemp(Builder);
diff --git a/Engine/Plugins/Media/PixelStreaming2/Source/PixelStreaming2RTC/Private/EpicRtcDataTrack.cpp b/Engine/Plugins/Media/PixelStreaming2/Source/PixelStreaming2RTC/Private/EpicRtcDataTrack.cpp
--- a/Engine/Plugins/Media/PixelStreaming2/Source/PixelStreaming2RTC/Private/EpicRtcDataTrack.cpp
+++ b/Engine/Plugins/Media/PixelStreaming2/Source/PixelStreaming2RTC/Private/EpicRtcDataTrack.cpp
@@ -121,10 +121,10 @@ namespace UE::PixelStreaming2
 			PrependData(Builder);
 
 			// Write message type
-			Builder.Insert(Type);
+			Builder.Insert(Forward<uint8>(Type));
 
 			// Write size of payload
-			Builder.Insert(DataSize);
+			Builder.Insert(Forward<const int32>(DataSize));
 
 			// Write the data bytes payload
 			Builder.Serialize(DataBytes.GetData() + BytesTransmitted, BytesToTransmit);
@@ -167,12 +167,12 @@ namespace UE::PixelStreaming2
 			return;
 		}
 
-		uint16 StringLength = static_cast<uint16>(ValueSize(PlayerId));
+		uint16 StringLength = static_cast<uint16>(GetByteSizeOf(Forward<FString>(FString(PlayerId))));
 
-		Builder.Buffer.SetNum(Builder.Buffer.Num() + ValueSize(Type) + ValueSize(StringLength) + StringLength);
+		Builder.Buffer.SetNum(Builder.Buffer.Num() + GetByteSizeOf(Forward<uint8>(Type)) + GetByteSizeOf(Forward<uint16>(StringLength)) + StringLength);
 
-		Builder.Insert(Type);
-		Builder.Insert(StringLength);
-		Builder.Insert<FString>(FString(PlayerId));
+		Builder.Insert(Forward<uint8>(Type));
+		Builder.Insert(Forward<uint16>(StringLength));
+		Builder.Insert(Forward<FString>(FString(PlayerId)));
 	}
 } // namespace UE::PixelStreaming2
