[PS2] Fix Max bitrate not being set.

Max bitrate may go below the start bitrate. To fix, this lowers the start bitrate if it is lower than the max bitrate.

Additionally, sets the initial max bitrate to be extremely high. This is because webrtc caps the max bitrate to the initial setting. By setting it initially to be high, it will go to the real value when SetConnectionRates is called.

diff --git a/Engine/Plugins/Media/PixelStreaming2/Source/PixelStreaming2RTC/Private/EpicRtcStreamer.cpp b/Engine/Plugins/Media/PixelStreaming2/Source/PixelStreaming2RTC/Private/EpicRtcStreamer.cpp
--- a/Engine/Plugins/Media/PixelStreaming2/Source/PixelStreaming2RTC/Private/EpicRtcStreamer.cpp
+++ b/Engine/Plugins/Media/PixelStreaming2/Source/PixelStreaming2RTC/Private/EpicRtcStreamer.cpp
@@ -406,6 +406,24 @@ namespace UE::PixelStreaming2
 		SendAllPlayersMessage(MessageName, Descriptor);
 	}
 
+	EpicRtcBitrate FEpicRtcStreamer::GetBitrates() const
+	{
+		const int MinBitrate = UPixelStreaming2PluginSettings::CVarWebRTCMinBitrate.GetValueOnAnyThread();
+		const int MaxBitrate = UPixelStreaming2PluginSettings::CVarWebRTCMaxBitrate.GetValueOnAnyThread();
+		const int StartBitrate = FMath::Clamp(UPixelStreaming2PluginSettings::CVarWebRTCStartBitrate.GetValueOnAnyThread(), MinBitrate, MaxBitrate);
+
+		EpicRtcBitrate Bitrates = {
+			._minBitrateBps = MinBitrate,
+			._hasMinBitrateBps = true,
+			._maxBitrateBps = MaxBitrate,
+			._hasMaxBitrateBps = true,
+			._startBitrateBps = StartBitrate,
+			._hasStartBitrateBps = true
+		};
+
+		return Bitrates;
+	}
+
 	IPixelStreaming2Streamer::FPreConnectionEvent& FEpicRtcStreamer::OnPreConnection()
 	{
 		return StreamingPreConnectionEvent;
@@ -608,16 +626,7 @@ namespace UE::PixelStreaming2
 				return;
 			}
 
-			EpicRtcBitrate Bitrates = {
-				._minBitrateBps = UPixelStreaming2PluginSettings::CVarWebRTCMinBitrate.GetValueOnAnyThread(),
-				._hasMinBitrateBps = true,
-				._maxBitrateBps = UPixelStreaming2PluginSettings::CVarWebRTCMaxBitrate.GetValueOnAnyThread(),
-				._hasMaxBitrateBps = true,
-				._startBitrateBps = UPixelStreaming2PluginSettings::CVarWebRTCStartBitrate.GetValueOnAnyThread(),
-				._hasStartBitrateBps = true
-			};
-
-			ConnectionInterface->SetConnectionRates(Bitrates);
+			ConnectionInterface->SetConnectionRates(GetBitrates());
 		});
 	}
 
@@ -980,14 +989,7 @@ namespace UE::PixelStreaming2
 				}
 
 				UE_LOG(LogPixelStreaming2RTC, VeryVerbose, TEXT("FEpicRtcStreamer::OnSessionStateUpdate State=Connected"));
-				EpicRtcBitrate Bitrate = {
-					._minBitrateBps = UPixelStreaming2PluginSettings::CVarWebRTCMinBitrate.GetValueOnAnyThread(),
-					._hasMinBitrateBps = true,
-					._maxBitrateBps = UPixelStreaming2PluginSettings::CVarWebRTCMaxBitrate.GetValueOnAnyThread(),
-					._hasMaxBitrateBps = true,
-					._startBitrateBps = UPixelStreaming2PluginSettings::CVarWebRTCStartBitrate.GetValueOnAnyThread(),
-					._hasStartBitrateBps = true
-				};
+				EpicRtcBitrate Bitrate = GetBitrates();
 
 				EpicRtcPortAllocator PortAllocator = {
 					._minPort = UPixelStreaming2PluginSettings::CVarWebRTCMinPort.GetValueOnAnyThread(),
@@ -1146,12 +1148,20 @@ namespace UE::PixelStreaming2
 			for (int i = 0; i < SimulcastParams.Num(); ++i)
 			{
 				const FPixelStreaming2SimulcastLayer& SpatialLayer = SimulcastParams[i];
+				const bool bIsSimulcast = SimulcastParams.Num() > 1;
 
-				if (SimulcastParams.Num() > 1)
+				if (bIsSimulcast)
 				{
 					Rids.Add(FUtf8String("simulcast") + FUtf8String::FromInt(SimulcastParams.Num() - i));
 				}
 
+				// WebRTC cannot set the bitrate outside the first initial biterate set by the VideoEncodingConfig.
+				// By setting a high value here, the real value can be set by SetConnectionRates which can fit within this range.
+				// Without this, changing the max bitrate at runtime will be capped at the initial max bitrate.
+				// SetConnectionRates below will set the real max bitrate.
+				constexpr uint32 InitialMinBitrate = 1'000;
+				constexpr uint32 InitialMaxBitrate = 1'000'000'000;
+
 				EpicRtcVideoEncodingConfig VideoEncodingConfig = {
 					// clang-format off
 					// TODO (Migration): RTCP-7027 Maybe bug in EpicRtc? Setting an rid if there's only one config results in no video
@@ -1161,9 +1171,9 @@ namespace UE::PixelStreaming2
 					// clang-format on
 					._scaleResolutionDownBy = SpatialLayer.Scaling,
 					._scalabilityMode = static_cast<EpicRtcVideoScalabilityMode>(GetEnumFromCVar<EScalabilityMode>(UPixelStreaming2PluginSettings::CVarEncoderScalabilityMode)), // HACK if the Enums become un-aligned
-					._minBitrate = (uint32_t)SpatialLayer.MinBitrate,
-					._maxBitrate = (uint32_t)SpatialLayer.MaxBitrate,
-					._maxFrameRate = (uint8_t)MaxFramerate
+					._minBitrate = bIsSimulcast ? static_cast<uint32>(SpatialLayer.MinBitrate) : InitialMinBitrate,
+					._maxBitrate = bIsSimulcast ? static_cast<uint32>(SpatialLayer.MaxBitrate) : InitialMaxBitrate,
+					._maxFrameRate = static_cast<uint8_t>(MaxFramerate)
 				};
 
 				VideoEncodingConfigs.Add(VideoEncodingConfig);
@@ -1275,16 +1285,7 @@ namespace UE::PixelStreaming2
 			ParticipantConnection->AddDataSource(DataSource);
 		}
 
-		EpicRtcBitrate Bitrates = {
-			._minBitrateBps = UPixelStreaming2PluginSettings::CVarWebRTCMinBitrate.GetValueOnAnyThread(),
-			._hasMinBitrateBps = true,
-			._maxBitrateBps = UPixelStreaming2PluginSettings::CVarWebRTCMaxBitrate.GetValueOnAnyThread(),
-			._hasMaxBitrateBps = true,
-			._startBitrateBps = UPixelStreaming2PluginSettings::CVarWebRTCStartBitrate.GetValueOnAnyThread(),
-			._hasStartBitrateBps = true
-		};
-
-		ParticipantConnection->SetConnectionRates(Bitrates);
+		ParticipantConnection->SetConnectionRates(GetBitrates());
 
 		ParticipantConnection->StartNegotiation();
 	}
diff --git a/Engine/Plugins/Media/PixelStreaming2/Source/PixelStreaming2RTC/Private/EpicRtcStreamer.h b/Engine/Plugins/Media/PixelStreaming2/Source/PixelStreaming2RTC/Private/EpicRtcStreamer.h
--- a/Engine/Plugins/Media/PixelStreaming2/Source/PixelStreaming2RTC/Private/EpicRtcStreamer.h
+++ b/Engine/Plugins/Media/PixelStreaming2/Source/PixelStreaming2RTC/Private/EpicRtcStreamer.h
@@ -100,6 +100,7 @@ namespace UE::PixelStreaming2
 		void OnFrameCapturerCreated();
 		void OnUIInteraction(FMemoryReader Ar);
 		void OnSendMessage(FString MessageName, FMemoryReader Ar);
+		EpicRtcBitrate GetBitrates() const;
 
 	private:
 		FString StreamerId;
