[PS2] Fix deadlock when closing a stream.

TThreadSafeMap could deadlock when it calls Empty on the Players. What would happen is the Apply and Empty both lock, and the destructors of the players would also lock a thread.

This fix moves the deletion to outside the lock in the TThreadSafeMap Empty function as well as some other methods.

diff --git a/Engine/Plugins/Media/PixelStreaming2/Source/PixelStreaming2/Internal/ThreadSafeMap.h b/Engine/Plugins/Media/PixelStreaming2/Source/PixelStreaming2/Internal/ThreadSafeMap.h
--- a/Engine/Plugins/Media/PixelStreaming2/Source/PixelStreaming2/Internal/ThreadSafeMap.h
+++ b/Engine/Plugins/Media/PixelStreaming2/Source/PixelStreaming2/Internal/ThreadSafeMap.h
@@ -63,9 +63,15 @@ namespace UE::PixelStreaming2
 		template <typename OtherSetAllocator>
 		TThreadSafeMap& operator=(TThreadSafeMap<KeyType, ValueType, OtherSetAllocator, KeyFuncs>&& Other)
 		{
-			FScopeLock Lock(&Mutex);
+			// Move Pairs so they they are destroyed outside of the ScopeLock
+			TSet<TPair<KeyType, ValueType>, KeyFuncs, SetAllocator> MovedPairs;
+			{
+				FScopeLock Lock(&Mutex);
+				MovedPairs = MoveTemp(Super::Pairs);
 
-			(Super&)* this = MoveTemp(Other);
+				(Super&)* this = MoveTemp(Other);
+			}
+			MovedPairs.Empty();
 			return *this;
 		}
 
@@ -73,9 +79,15 @@ namespace UE::PixelStreaming2
 		template <typename OtherSetAllocator>
 		TThreadSafeMap& operator=(const TThreadSafeMap<KeyType, ValueType, OtherSetAllocator, KeyFuncs>& Other)
 		{
-			FScopeLock Lock(&Mutex);
+			// Move Pairs so they they are destroyed outside of the ScopeLock
+			TSet<TPair<KeyType, ValueType>, KeyFuncs, SetAllocator> MovedPairs;
+			{
+				FScopeLock Lock(&Mutex);
+				MovedPairs = MoveTemp(Super::Pairs);
 
-			(Super&)* this = Other;
+				(Super&)* this = Other;
+			}
+			MovedPairs.Empty();
 			return *this;
 		}
 
@@ -258,9 +270,13 @@ namespace UE::PixelStreaming2
 
 		FORCEINLINE void Empty(int32 ExpectedNumElements = 0)
 		{
-			FScopeLock Lock(&Mutex);
-
-			return static_cast<Super*>(this)->Empty(ExpectedNumElements);
+			// Move Pairs so they they are destroyed outside of the ScopeLock
+			TSet<TPair<KeyType, ValueType>, KeyFuncs, SetAllocator> MovedPairs;
+			{
+				FScopeLock Lock(&Mutex);
+				MovedPairs = MoveTemp(Super::Pairs);
+			}
+			MovedPairs.Empty(ExpectedNumElements);
 		}
 
 		template <typename InSetKeyFuncs, typename InSetAllocator>
