[PS2] Fix use after free in Signalling server.

Fix was originally applied to PS1 in UE-241253.

This appears to resolve a crash that would happen when running the Pixel Streaming 2 Tests.

diff --git a/Engine/Plugins/Media/PixelStreaming2/Source/PixelStreaming2Servers/Private/SignallingServer.cpp b/Engine/Plugins/Media/PixelStreaming2/Source/PixelStreaming2Servers/Private/SignallingServer.cpp
--- a/Engine/Plugins/Media/PixelStreaming2/Source/PixelStreaming2Servers/Private/SignallingServer.cpp
+++ b/Engine/Plugins/Media/PixelStreaming2/Source/PixelStreaming2Servers/Private/SignallingServer.cpp
@@ -4,6 +4,7 @@
 #include "ServerUtils.h"
 #include "Logging.h"
 #include "Misc/Paths.h"
+#include "WebSocketServerWrapper.h"
 
 namespace UE::PixelStreaming2Servers
 {
@@ -43,22 +44,14 @@ namespace UE::PixelStreaming2Servers
 		}
 		else
 		{
-			bool bConnected = Probe->Probe();
+			const bool bConnected = Probe && Probe->Probe();
 			if (bConnected)
 			{
 				// Close the websocket connection so others can use it
 				Probe->Close();
 				Probe.Reset();
-
-				// Even after closing the client WS of the probe above it will take another tick to remove the connection
-				// so we manually remove the probe connection here to prevent any erroneous number of streamers being reported.
-				uint16 ProbeId = -1;
-				bool   bHasProbe = StreamersWS->GetFirstConnection(ProbeId);
-				if (bHasProbe)
-				{
-					StreamersWS->Close(ProbeId);
-				}
-
+				// Note: even after closing the client WS of the probe above it will take another tick to remove the connection
+				// from the ws server so if you query number of streamers during this time we do have to remove the probe from that count
 				return true;
 			}
 			else
@@ -487,7 +480,13 @@ namespace UE::PixelStreaming2Servers
 	{
 		if (StreamersWS)
 		{
-			OnNumStreamersReceived(StreamersWS->Count());
+			int NConnections = StreamersWS->Count();
+			// If the probe is currently connected, it is not a streamer, so don't count it.
+			if (Probe && Probe->IsConnected()) 
+			{
+				NConnections = NConnections - 1;
+			}
+			OnNumStreamersReceived(NConnections);
 		}
 		else
 		{
