[PS2] Refactor: Stats. Cleaned up stat related classes. Added support for string and boolean based stats. Added tracking of audio and video codec stats. Fixed calculation for mean packet send delay.

diff --git a/Engine/Plugins/Media/PixelStreaming2/Source/PixelStreaming2RTC/Internal/RTCStatsCollector.h b/Engine/Plugins/Media/PixelStreaming2/Source/PixelStreaming2RTC/Internal/RTCStatsCollector.h
--- a/Engine/Plugins/Media/PixelStreaming2/Source/PixelStreaming2RTC/Internal/RTCStatsCollector.h
+++ b/Engine/Plugins/Media/PixelStreaming2/Source/PixelStreaming2RTC/Internal/RTCStatsCollector.h
@@ -4,6 +4,8 @@
 
 #include "Containers/Map.h"
 #include "Containers/UnrealString.h"
+#include "Misc/Optional.h"
+#include "Misc/TVariant.h"
 #include "Templates/SharedPointer.h"
 #include "UObject/NameTypes.h"
 
@@ -12,66 +14,121 @@
 namespace UE::PixelStreaming2
 {
 	class FStats;
-	
+
 	namespace RTCStatCategories
 	{
 		const FString LocalVideoTrack = TEXT("video-track-outbound");
 		const FString LocalAudioTrack = TEXT("audio-track-outbound");
 		const FString VideoSource = TEXT("video-source");
 		const FString AudioSource = TEXT("audio-source");
+		const FString VideoCodec = TEXT("video-codec");
+		const FString AudioCodec = TEXT("audio-codec");
 		const FString DataChannel = TEXT("data-channel");
 		const FString RemoteVideoTrack = TEXT("video-track-inbound");
 		const FString RemoteAudioTrack = TEXT("audio-track-inbound");
 		const FString CandidatePair = TEXT("candidate-pair");
 	} // namespace RTCStatCategories
 
-	struct PIXELSTREAMING2RTC_API FStatData
+	enum PIXELSTREAMING2RTC_API EDisplayFlags : uint8
 	{
-	public:
-		enum EDisplayFlags : uint8
-		{
-			HIDDEN = 0,
-			TEXT = 1 << 0,
-			GRAPH = 1 << 1,
-		};
+		HIDDEN = 0,
+		TEXT = 1 << 0,
+		GRAPH = 1 << 1,
+	};
+
+	struct PIXELSTREAMING2RTC_API FStatConfig
+	{
+		FName			 Name;
+		TOptional<FName> Alias;
+		EDisplayFlags	 DisplayFlags = EDisplayFlags::TEXT;
+	};
+
+	using FStatVariant = TVariant<TYPE_OF_NULLPTR, double, FString, bool>;
 
+	class PIXELSTREAMING2RTC_API FStat
+	{
 	public:
-		FStatData(FName InStatName, double InStatValue, int InNDecimalPlacesToPrint, bool bInSmooth = false)
-			: StatName(InStatName)
-			, StatValue(InStatValue)
-			, NDecimalPlacesToPrint(InNDecimalPlacesToPrint)
-			, bSmooth(bInSmooth)
-		{
-		}
+		FStat() = default;
+
+		FStat(FStatConfig Config, double InitialValue, int NDecimalPlacesToPrint = 0, bool bSmooth = false);
+		FStat(FStatConfig Config, FString InitialValue);
+		FStat(FStatConfig Config, bool bInitialValue);
+		FStat(const FStat& Other);
+
+		bool IsNumeric() const;
+
+		bool IsTextual() const;
+
+		bool IsBoolean() const;
+
+		FString ToString();
 
-		bool operator==(const FStatData& Other) const
+		bool SetValue(FStatVariant ValueVariant);
+
+		template <typename T>
+		T GetValue()
 		{
-			return Equals(Other);
+			checkf(false, TEXT("Trying to get value with incompatible container. Please use FString, double or bool!"));
 		}
 
-		bool Equals(const FStatData& Other) const
+		template <>
+		FString GetValue();
+
+		template <>
+		double GetValue();
+
+		template <>
+		bool GetValue();
+
+		template <typename T>
+		T GetPrevValue()
 		{
-			return StatName == Other.StatName;
+			checkf(false, TEXT("Trying to get value with incompatible container. Please use FString, double or bool!"));
 		}
 
-		bool IsHidden() { return DisplayFlags == EDisplayFlags::HIDDEN; }
-		bool ShouldGraph() { return DisplayFlags & EDisplayFlags::GRAPH; }
-		bool ShouldDisplayText() { return DisplayFlags & EDisplayFlags::TEXT; }
-
-		FName			 StatName;
-		double			 StatValue;
-		int				 NDecimalPlacesToPrint;
-		bool			 bSmooth;
-		double			 LastEMA = 0;
-		int				 NumSamples = 0;
-		uint8			 DisplayFlags = EDisplayFlags::TEXT; // Some stats we only wish to store or broadcast, but not display
-		TOptional<FName> Alias;								 // Some stats need an alias that they are stored by/queried by to disambiguate them from other stats
+		template <>
+		FString GetPrevValue();
+
+		template <>
+		double GetPrevValue();
+
+		template <>
+		bool GetPrevValue();
+
+		bool operator==(const FStat& Other) const;
+
+		bool IsHidden();
+
+		bool ShouldGraph();
+
+		bool ShouldDisplayText();
+
+		FName GetName() const;
+
+		FName GetDisplayName() const;
+
+	private:
+		double CalcMA(double InPrevAvg, int InNumSamples, double InValue);
+		double CalcEMA(double InPrevAvg, int InNumSamples, double InValue);
+
+	protected:
+		FName			 Name;
+		EDisplayFlags	 DisplayFlags = EDisplayFlags::TEXT;
+		TOptional<FName> Alias;
+
+		// variable used for calculating and displaying numeric stat data
+		int	 NDecimalPlacesToPrint = 0;
+		bool bSmooth = false;
+		int	 NumSamples = 0;
+
+		FStatVariant StatVariant;
+		FStatVariant PrevStatVariant;
 	};
 
-	FORCEINLINE uint32 GetTypeHash(const FStatData& Obj)
+	FORCEINLINE uint32 GetTypeHash(const FStat& Obj)
 	{
 		// From UnrealString.h
-		return GetTypeHash(Obj.StatName);
+		return GetTypeHash(Obj.GetName());
 	}
 
 	class PIXELSTREAMING2RTC_API FRTCStatsCollector
@@ -89,72 +146,24 @@ namespace UE::PixelStreaming2
 		void Process(const EpicRtcConnectionStats& InStats);
 
 	private:
-		/**
-		 * ---------- FRTCTrackedStat -------------------
-		 * Tracks a stat from WebRTC. Stores the current value and previous value.
-		 */
-		class FRTCTrackedStat
-		{
-		public:
-			FRTCTrackedStat(FName StatName, int NDecimalPlaces)
-				: LatestStat(StatName, 0.0, NDecimalPlaces)
-			{
-			}
-
-			FRTCTrackedStat(FName StatName, int NDecimalPlaces, uint8 DisplayFlags)
-				: LatestStat(StatName, 0.0, NDecimalPlaces)
-			{
-				LatestStat.DisplayFlags = DisplayFlags;
-			}
-
-			FRTCTrackedStat(FName StatName, FName Alias, int NDecimalPlaces, uint8 DisplayFlags);
-
-			double			 CalculateDelta(double Period) const;
-			double			 Average() const;
-			const FStatData& GetLatestStat() const { return LatestStat; }
-			void			 SetLatestValue(double InValue);
-
-		private:
-			FStatData LatestStat;
-			double	  PrevValue = 0.0;
-		};
-
 		class FStatsSink
 		{
 		public:
 			FStatsSink(FName InCategory);
 			virtual ~FStatsSink() = default;
-			virtual void Add(FName StatName, int NDecimalPlaces)
-			{
-				Stats.Add(StatName, FRTCTrackedStat(StatName, NDecimalPlaces, FStatData::EDisplayFlags::TEXT));
-			}
 
-			virtual void AddAliased(FName StatName, FName AliasedName, int NDecimalPlaces, uint8 DisplayFlags);
-			virtual void AddNonRendered(FName StatName)
-			{
-				Stats.Add(StatName, FRTCTrackedStat(StatName, 2, FStatData::EDisplayFlags::HIDDEN));
-			}
-			virtual void AddStatCalculator(const TFunction<TOptional<FStatData>(FStatsSink&, double)>& Calculator)
-			{
-				Calculators.Add(Calculator);
-			}
-			// returns true if the value is worth storing (false if it started and remains zero)
-			virtual bool			 UpdateValue(double NewValue, FRTCTrackedStat* SetValueHere);
-			virtual FRTCTrackedStat* Get(const FName& StatName)
+			virtual FStat* Get(const FName& StatName)
 			{
 				return Stats.Find(StatName);
 			}
-			virtual FStatData* GetCalculatedStat(const FName& StatName)
-			{
-				return CalculatedStats.Find(StatName);
-			}
+
 			virtual void PostProcess(FStats* PSStats, const FString& PeerId, double SecondsDelta);
 
+		protected:
 			// Stats that are stored as is.
-			TMap<FName, FRTCTrackedStat> Stats;
+			TMap<FName, FStat> Stats;
 			// Stats we calculate based on the stats map above. This calculation is done in FStatsSink::PostProcess by the `Calculators` below.
-			TMap<FName, FStatData>										 CalculatedStats;
-			TArray<TFunction<TOptional<FStatData>(FStatsSink&, double)>> Calculators;
+			TArray<TFunction<TOptional<FStat>(FStatsSink&, double)>> Calculators;
 
 		protected:
 			FName Category;
@@ -192,6 +201,14 @@ namespace UE::PixelStreaming2
 			void Process(const EpicRtcVideoSourceStats& InStats, const FString& PeerId, double SecondsDelta);
 		};
 
+		class FVideoCodecStatsSink : public FStatsSink
+		{
+		public:
+			FVideoCodecStatsSink(FName Category);
+			virtual ~FVideoCodecStatsSink() = default;
+			void Process(const EpicRtcCodecStats& InStats, const FString& PeerId, double SecondsDelta);
+		};
+
 		class FAudioSourceStatsSink : public FStatsSink
 		{
 		public:
@@ -200,6 +217,14 @@ namespace UE::PixelStreaming2
 			void Process(const EpicRtcAudioSourceStats& InStats, const FString& PeerId, double SecondsDelta);
 		};
 
+		class FAudioCodecStatsSink : public FStatsSink
+		{
+		public:
+			FAudioCodecStatsSink(FName Category);
+			virtual ~FAudioCodecStatsSink() = default;
+			void Process(const EpicRtcCodecStats& InStats, const FString& PeerId, double SecondsDelta);
+		};
+
 		/**
 		 * ---------- FDataChannelSink -------------------
 		 */
@@ -233,7 +258,10 @@ namespace UE::PixelStreaming2
 		TMap<uint64, TMap<uint32, TUniquePtr<FRTPRemoteTrackStatsSink>>> RemoteAudioTrackSinks;
 
 		TMap<uint64, TUniquePtr<FVideoSourceStatsSink>> VideoSourceSinks;
+		TMap<uint64, TUniquePtr<FVideoCodecStatsSink>>	VideoCodecSinks;
+
 		TMap<uint64, TUniquePtr<FAudioSourceStatsSink>> AudioSourceSinks;
+		TMap<uint64, TUniquePtr<FAudioCodecStatsSink>>	AudioCodecSinks;
 
 		TMap<uint64, TUniquePtr<FDataTrackStatsSink>> DataTrackSinks;
 
diff --git a/Engine/Plugins/Media/PixelStreaming2/Source/PixelStreaming2RTC/Private/EpicRtcStreamer.cpp b/Engine/Plugins/Media/PixelStreaming2/Source/PixelStreaming2RTC/Private/EpicRtcStreamer.cpp
--- a/Engine/Plugins/Media/PixelStreaming2/Source/PixelStreaming2RTC/Private/EpicRtcStreamer.cpp
+++ b/Engine/Plugins/Media/PixelStreaming2/Source/PixelStreaming2RTC/Private/EpicRtcStreamer.cpp
@@ -374,13 +374,13 @@ namespace UE::PixelStreaming2
 			PSStats->RemoveAllApplicationStats();
 
 			// Re-add the cvar based application stats so that they're at the top
-			PSStats->StoreApplicationStat(FStatData(FName(TEXT("PixelStreaming2.Encoder.MinQuality")), UPixelStreaming2PluginSettings::CVarEncoderMinQuality.GetValueOnAnyThread(), 0));
-			PSStats->StoreApplicationStat(FStatData(FName(TEXT("PixelStreaming2.Encoder.MaxQuality")), UPixelStreaming2PluginSettings::CVarEncoderMaxQuality.GetValueOnAnyThread(), 0));
-			PSStats->StoreApplicationStat(FStatData(FName(TEXT("PixelStreaming2.Encoder.KeyframeInterval (frames)")), UPixelStreaming2PluginSettings::CVarEncoderKeyframeInterval.GetValueOnAnyThread(), 0));
-			PSStats->StoreApplicationStat(FStatData(FName(TEXT("PixelStreaming2.WebRTC.Fps")), UPixelStreaming2PluginSettings::CVarWebRTCFps.GetValueOnAnyThread(), 0));
-			PSStats->StoreApplicationStat(FStatData(FName(TEXT("PixelStreaming2.WebRTC.StartBitrate")), UPixelStreaming2PluginSettings::CVarWebRTCStartBitrate.GetValueOnAnyThread(), 0));
-			PSStats->StoreApplicationStat(FStatData(FName(TEXT("PixelStreaming2.WebRTC.MinBitrate")), UPixelStreaming2PluginSettings::CVarWebRTCMinBitrate.GetValueOnAnyThread(), 0));
-			PSStats->StoreApplicationStat(FStatData(FName(TEXT("PixelStreaming2.WebRTC.MaxBitrate")), UPixelStreaming2PluginSettings::CVarWebRTCMaxBitrate.GetValueOnAnyThread(), 0));
+			PSStats->StoreApplicationStat(FStat({ .Name = FName(TEXT("PixelStreaming2.Encoder.MinQuality")) }, UPixelStreaming2PluginSettings::CVarEncoderMinQuality.GetValueOnAnyThread(), 0));
+			PSStats->StoreApplicationStat(FStat({ .Name = FName(TEXT("PixelStreaming2.Encoder.MaxQuality")) }, UPixelStreaming2PluginSettings::CVarEncoderMaxQuality.GetValueOnAnyThread(), 0));
+			PSStats->StoreApplicationStat(FStat({ .Name = FName(TEXT("PixelStreaming2.Encoder.KeyframeInterval (frames)")) }, UPixelStreaming2PluginSettings::CVarEncoderKeyframeInterval.GetValueOnAnyThread(), 0));
+			PSStats->StoreApplicationStat(FStat({ .Name = FName(TEXT("PixelStreaming2.WebRTC.Fps")) }, UPixelStreaming2PluginSettings::CVarWebRTCFps.GetValueOnAnyThread(), 0));
+			PSStats->StoreApplicationStat(FStat({ .Name = FName(TEXT("PixelStreaming2.WebRTC.StartBitrate")) }, UPixelStreaming2PluginSettings::CVarWebRTCStartBitrate.GetValueOnAnyThread(), 0));
+			PSStats->StoreApplicationStat(FStat({ .Name = FName(TEXT("PixelStreaming2.WebRTC.MinBitrate")) }, UPixelStreaming2PluginSettings::CVarWebRTCMinBitrate.GetValueOnAnyThread(), 0));
+			PSStats->StoreApplicationStat(FStat({ .Name = FName(TEXT("PixelStreaming2.WebRTC.MaxBitrate")) }, UPixelStreaming2PluginSettings::CVarWebRTCMaxBitrate.GetValueOnAnyThread(), 0));
 		}
 	}
 
diff --git a/Engine/Plugins/Media/PixelStreaming2/Source/PixelStreaming2RTC/Private/EpicRtcVideoEncoder.cpp b/Engine/Plugins/Media/PixelStreaming2/Source/PixelStreaming2RTC/Private/EpicRtcVideoEncoder.cpp
--- a/Engine/Plugins/Media/PixelStreaming2/Source/PixelStreaming2RTC/Private/EpicRtcVideoEncoder.cpp
+++ b/Engine/Plugins/Media/PixelStreaming2/Source/PixelStreaming2RTC/Private/EpicRtcVideoEncoder.cpp
@@ -64,8 +64,22 @@ namespace UE::PixelStreaming2
 	template <std::derived_from<FVideoResource> TVideoResource>
 	EpicRtcStringView TEpicRtcVideoEncoder<TVideoResource>::GetName() const
 	{
-		static FUtf8String Name("PixelStreamingVideoEncoder");
-		return ToEpicRtcStringView(Name);
+		static FUtf8String Utf8ImplementationString = [this]() {
+			FString EncoderLocation = TEXT("Unknown");
+			if constexpr (std::is_same_v<TVideoResource, FVideoResourceRHI>)
+			{
+				EncoderLocation = TEXT("GPU");
+			}
+			else if constexpr (std::is_same_v<TVideoResource, FVideoResourceCPU>)
+			{
+				EncoderLocation = TEXT("CPU");
+			}
+
+			FString		ImplementationString = FString::Printf(TEXT("PixelStreamingVideoEncoder(%s)(%s)"), *ToString(CodecInfo->GetCodec()), *EncoderLocation);
+			return FUtf8String(ImplementationString);
+		}();
+		
+		return ToEpicRtcStringView(Utf8ImplementationString);
 	}
 
 	template <std::derived_from<FVideoResource> TVideoResource>
diff --git a/Engine/Plugins/Media/PixelStreaming2/Source/PixelStreaming2RTC/Private/PixelStreaming2StatNames.cpp b/Engine/Plugins/Media/PixelStreaming2/Source/PixelStreaming2RTC/Private/PixelStreaming2StatNames.cpp
--- a/Engine/Plugins/Media/PixelStreaming2/Source/PixelStreaming2RTC/Private/PixelStreaming2StatNames.cpp
+++ b/Engine/Plugins/Media/PixelStreaming2/Source/PixelStreaming2RTC/Private/PixelStreaming2StatNames.cpp
@@ -64,6 +64,12 @@ namespace PixelStreaming2StatNames
 	const FName AvailableIncomingBitrate = FName(TEXT("availableIncomingBitrate"));
 	const FName RetransmittedBytesReceived = FName(TEXT("retransmittedBytesReceived"));
 	const FName RetransmittedPacketsReceived = FName(TEXT("retransmittedPacketsReceived"));
+	const FName PacketsSent = FName(TEXT("packetsSent"));
+	const FName EncoderImplementation = FName(TEXT("encoderImplementation"));
+	const FName MimeType = FName(TEXT("mimeType"));
+	const FName Channels = FName(TEXT("channels"));
+	const FName ClockRate = FName(TEXT("clockRate"));
+
 
 	// Data-channel stats
 	const FName DataChannelBytesSent = FName(TEXT("data-channel-bytesSent"));
diff --git a/Engine/Plugins/Media/PixelStreaming2/Source/PixelStreaming2RTC/Private/RTCStatsCollector.cpp b/Engine/Plugins/Media/PixelStreaming2/Source/PixelStreaming2RTC/Private/RTCStatsCollector.cpp
--- a/Engine/Plugins/Media/PixelStreaming2/Source/PixelStreaming2RTC/Private/RTCStatsCollector.cpp
+++ b/Engine/Plugins/Media/PixelStreaming2/Source/PixelStreaming2RTC/Private/RTCStatsCollector.cpp
@@ -11,6 +11,312 @@
 
 namespace UE::PixelStreaming2
 {
+	/**
+	 * ---------- FStat ----------
+	 */
+
+	FStat::FStat(FStatConfig Config, double InitialValue, int NDecimalPlacesToPrint, bool bSmooth)
+		: Name(Config.Name)
+		, DisplayFlags(Config.DisplayFlags)
+		, Alias(Config.Alias)
+		, NDecimalPlacesToPrint(NDecimalPlacesToPrint)
+		, bSmooth(bSmooth)
+		, StatVariant(TInPlaceType<double>(), InitialValue)
+	{
+	}
+
+	FStat::FStat(FStatConfig Config, FString InitialValue)
+		: Name(Config.Name)
+		, DisplayFlags(Config.DisplayFlags)
+		, Alias(Config.Alias)
+		, StatVariant(TInPlaceType<FString>(), InitialValue)
+	{
+		checkf(!(Config.DisplayFlags & EDisplayFlags::GRAPH), TEXT("Text based stats cannot be graphed"));
+	}
+
+	FStat::FStat(FStatConfig Config, bool bInitialValue)
+		: Name(Config.Name)
+		, DisplayFlags(Config.DisplayFlags)
+		, Alias(Config.Alias)
+		, StatVariant(TInPlaceType<bool>(), bInitialValue)
+	{
+		checkf(!(Config.DisplayFlags & EDisplayFlags::GRAPH), TEXT("Boolean based stats cannot be graphed"));
+	}
+
+	FStat::FStat(const FStat& Other)
+		: Name(Other.Name)
+		, DisplayFlags(Other.DisplayFlags)
+		, Alias(Other.Alias)
+		, NDecimalPlacesToPrint(Other.NDecimalPlacesToPrint)
+		, bSmooth(Other.bSmooth)
+		, StatVariant(Other.StatVariant)
+	{
+	}
+
+	bool FStat::IsNumeric() const
+	{
+		return StatVariant.GetIndex() == FStatVariant::IndexOfType<double>();
+	}
+
+	bool FStat::IsTextual() const
+	{
+		return StatVariant.GetIndex() == FStatVariant::IndexOfType<FString>();
+	}
+
+	bool FStat::IsBoolean() const
+	{
+		return StatVariant.GetIndex() == FStatVariant::IndexOfType<bool>();
+	}
+
+	FString FStat::ToString()
+	{
+		switch (StatVariant.GetIndex())
+		{
+			case FStatVariant::IndexOfType<FString>():
+			{
+				return StatVariant.Get<FString>();
+			}
+			case FStatVariant::IndexOfType<double>():
+			{
+				return FString::Printf(TEXT("%.*f"), NDecimalPlacesToPrint, StatVariant.Get<double>());
+			}
+			case FStatVariant::IndexOfType<bool>():
+			{
+				return FString::Printf(TEXT("%s"), StatVariant.Get<bool>() ? TEXT("true") : TEXT("false"));
+			}
+			default:
+				checkNoEntry();
+				return TEXT("");
+		}
+	}
+
+	bool FStat::SetValue(FStatVariant ValueVariant)
+	{
+		if (ValueVariant.GetIndex() == FStatVariant::IndexOfType<TYPE_OF_NULLPTR>())
+		{
+			return false;
+		}
+
+		if (ValueVariant.GetIndex() != StatVariant.GetIndex())
+		{
+			FString ValueVariantType = TEXT("TYPE_OF_NULLPTR");
+			switch (ValueVariant.GetIndex())
+			{
+				case FStatVariant::IndexOfType<FString>():
+				{
+					ValueVariantType = TEXT("FString");
+					break;
+				}
+				case FStatVariant::IndexOfType<double>():
+				{
+					ValueVariantType = TEXT("double");
+					break;
+				}
+				case FStatVariant::IndexOfType<bool>():
+				{
+					ValueVariantType = TEXT("bool");
+					break;
+				}
+				default:
+					checkNoEntry();
+					break;
+			}
+
+			FString StatVariantType = TEXT("TYPE_OF_NULLPTR");
+			switch (StatVariant.GetIndex())
+			{
+				case FStatVariant::IndexOfType<FString>():
+				{
+					StatVariantType = TEXT("FString");
+					break;
+				}
+				case FStatVariant::IndexOfType<double>():
+				{
+					StatVariantType = TEXT("double");
+					break;
+				}
+				case FStatVariant::IndexOfType<bool>():
+				{
+					StatVariantType = TEXT("bool");
+					break;
+				}
+				default:
+					checkNoEntry();
+					break;
+			}
+
+			UE_LOGFMT(LogPixelStreaming2RTC, Warning, "Attempted to assign a {0} to a {1} stat!. The operation wasn't successful!", ValueVariantType, StatVariantType);
+			return false;
+		}
+
+		PrevStatVariant = StatVariant;
+
+		switch (ValueVariant.GetIndex())
+		{
+			case FStatVariant::IndexOfType<FString>():
+			{
+				FString PrevValue = StatVariant.Get<FString>();
+				FString NewValue = ValueVariant.Get<FString>();
+
+				StatVariant.Set<FString>(NewValue);
+
+				return PrevValue != NewValue;
+			}
+			case FStatVariant::IndexOfType<double>():
+			{
+				double PrevValue = StatVariant.Get<double>();
+				double NewValue = ValueVariant.Get<double>();
+
+				if (bSmooth)
+				{
+					const int MaxSamples = 60;
+					const int NumSamplesToUse = FGenericPlatformMath::Min(MaxSamples, NumSamples + 1);
+
+					if (NumSamplesToUse < MaxSamples)
+					{
+						NewValue = CalcMA(PrevValue, NumSamples - 1, NewValue);
+					}
+					else
+					{
+						NewValue = CalcEMA(PrevValue, NumSamples - 1, NewValue);
+					}
+					NumSamples++;
+				}
+
+				StatVariant.Set<double>(NewValue);
+
+				return PrevValue != NewValue;
+			}
+			case FStatVariant::IndexOfType<bool>():
+			{
+				bool PrevValue = StatVariant.Get<bool>();
+				bool NewValue = ValueVariant.Get<bool>();
+
+				StatVariant.Set<bool>(NewValue);
+
+				return PrevValue != NewValue;
+			}
+			default:
+				checkNoEntry();
+				return false;
+		}
+	}
+
+	template <>
+	FString FStat::GetValue()
+	{
+		if (!IsTextual())
+		{
+			checkf(false, TEXT("Tried to get a string value from a non-string stat!"));
+			return TEXT("");
+		}
+
+		return StatVariant.Get<FString>();
+	}
+
+	template <>
+	double FStat::GetValue()
+	{
+		if (!IsNumeric())
+		{
+			checkf(false, TEXT("Tried to get a numeric value from a non-numeric stat!"));
+			return -1.f;
+		}
+
+		return StatVariant.Get<double>();
+	}
+
+	template <>
+	bool FStat::GetValue()
+	{
+		if (!IsBoolean())
+		{
+			checkf(false, TEXT("Tried to get a boolean value from a non-boolean stat!"));
+			return false;
+		}
+
+		return StatVariant.Get<bool>();
+	}
+
+	template <>
+	FString FStat::GetPrevValue()
+	{
+		if (PrevStatVariant.GetIndex() != FStatVariant::IndexOfType<FString>())
+		{
+			checkf(false, TEXT("Tried to get a string value from a non-string stat!"));
+			return TEXT("");
+		}
+
+		return PrevStatVariant.Get<FString>();
+	}
+
+	template <>
+	double FStat::GetPrevValue()
+	{
+		if (!IsNumeric())
+		{
+			checkf(false, TEXT("Tried to get a numeric value from a non-numeric stat!"));
+			return -1.f;
+		}
+
+		return PrevStatVariant.Get<double>();
+	}
+
+	template <>
+	bool FStat::GetPrevValue()
+	{
+		if (!IsBoolean())
+		{
+			checkf(false, TEXT("Tried to get a boolean value from a non-boolean stat!"));
+			return false;
+		}
+
+		return PrevStatVariant.Get<bool>();
+	}
+
+	bool FStat::operator==(const FStat& Other) const
+	{
+		return Name == Other.Name;
+	}
+
+	bool FStat::IsHidden()
+	{
+		return DisplayFlags == EDisplayFlags::HIDDEN;
+	}
+
+	bool FStat::ShouldGraph()
+	{
+		return DisplayFlags & EDisplayFlags::GRAPH;
+	}
+
+	bool FStat::ShouldDisplayText()
+	{
+		return DisplayFlags & EDisplayFlags::TEXT;
+	}
+
+	FName FStat::GetName() const
+	{
+		return Name;
+	}
+
+	FName FStat::GetDisplayName() const
+	{
+		return Alias.Get(Name);
+	}
+
+	double FStat::CalcMA(double InPrevAvg, int InNumSamples, double InValue)
+	{
+		const double Result = InNumSamples * InPrevAvg + InValue;
+		return Result / (InPrevAvg + 1.0);
+	}
+
+	double FStat::CalcEMA(double InPrevAvg, int InNumSamples, double InValue)
+	{
+		const double Mult = 2.0 / (InNumSamples + 1.0);
+		const double Result = (InValue - InPrevAvg) * Mult + InPrevAvg;
+		return Result;
+	}
+
 	TSharedPtr<FRTCStatsCollector> FRTCStatsCollector::Create(const FString& PlayerId)
 	{
 		TSharedPtr<FRTCStatsCollector> StatsCollector = TSharedPtr<FRTCStatsCollector>(new FRTCStatsCollector(PlayerId));
@@ -66,6 +372,15 @@ namespace UE::PixelStreaming2
 
 			VideoSourceSinks[i]->Process(LocalVideoTrackStats._source, AssociatedPlayerId, SecondsDelta);
 
+			// Process video codec stats
+			if (!VideoCodecSinks.Contains(i))
+			{
+				FName SinkName = FName(*FString::Printf(TEXT("%s [%u]"), *RTCStatCategories::VideoCodec, i));
+				VideoCodecSinks.Add(i, MakeUnique<FVideoCodecStatsSink>(SinkName));
+			}
+
+			VideoCodecSinks[i]->Process(LocalVideoTrackStats._codec, AssociatedPlayerId, SecondsDelta);
+
 			// Process video track rtp stats
 			if (!LocalVideoTrackSinks.Contains(i))
 			{
@@ -101,6 +416,15 @@ namespace UE::PixelStreaming2
 
 			AudioSourceSinks[i]->Process(LocalAudioTrackStats._source, AssociatedPlayerId, SecondsDelta);
 
+			// Process audio codec stats
+			if (!AudioCodecSinks.Contains(i))
+			{
+				FName SinkName = FName(*FString::Printf(TEXT("%s [%u]"), *RTCStatCategories::AudioCodec, i));
+				AudioCodecSinks.Add(i, MakeUnique<FAudioCodecStatsSink>(SinkName));
+			}
+
+			AudioCodecSinks[i]->Process(LocalAudioTrackStats._codec, AssociatedPlayerId, SecondsDelta);
+
 			// Process audio track rtp stats
 			if (!LocalAudioTrackSinks.Contains(i))
 			{
@@ -108,7 +432,7 @@ namespace UE::PixelStreaming2
 			}
 
 			TMap<uint32, TUniquePtr<FRTPLocalAudioTrackStatsSink>>& SsrcSinks = LocalAudioTrackSinks[i];
-			const EpicRtcLocalTrackRtpStats&						  RtpStats = LocalAudioTrackStats._rtp;
+			const EpicRtcLocalTrackRtpStats&						RtpStats = LocalAudioTrackStats._rtp;
 
 			if (!SsrcSinks.Contains(RtpStats._local._ssrc))
 			{
@@ -131,7 +455,7 @@ namespace UE::PixelStreaming2
 			}
 
 			TMap<uint32, TUniquePtr<FRTPRemoteTrackStatsSink>>& SsrcSinks = RemoteVideoTrackSinks[i];
-			const EpicRtcRemoteTrackRtpStats&					  RtpStats = RemoteVideoTrackStats._rtp;
+			const EpicRtcRemoteTrackRtpStats&					RtpStats = RemoteVideoTrackStats._rtp;
 
 			if (!SsrcSinks.Contains(RtpStats._local._ssrc))
 			{
@@ -154,7 +478,7 @@ namespace UE::PixelStreaming2
 			}
 
 			TMap<uint32, TUniquePtr<FRTPRemoteTrackStatsSink>>& SsrcSinks = RemoteAudioTrackSinks[i];
-			const EpicRtcRemoteTrackRtpStats&					  RtpStats = RemoteAudioTrackStats._rtp;
+			const EpicRtcRemoteTrackRtpStats&					RtpStats = RemoteAudioTrackStats._rtp;
 
 			if (!SsrcSinks.Contains(RtpStats._local._ssrc))
 			{
@@ -205,181 +529,183 @@ namespace UE::PixelStreaming2
 		: FStatsSink(InCategory)
 	{
 		// These stats will be extracted from the stat reports and emitted straight to screen
-		Add(PixelStreaming2StatNames::FirCount, 0);
-		Add(PixelStreaming2StatNames::PliCount, 0);
-		Add(PixelStreaming2StatNames::NackCount, 0);
-		Add(PixelStreaming2StatNames::RetransmittedBytesSent, 0);
-		Add(PixelStreaming2StatNames::TotalEncodeBytesTarget, 0);
-		Add(PixelStreaming2StatNames::KeyFramesEncoded, 0);
-		Add(PixelStreaming2StatNames::FrameWidth, 0);
-		Add(PixelStreaming2StatNames::FrameHeight, 0);
-		Add(PixelStreaming2StatNames::HugeFramesSent, 0);
-		Add(PixelStreaming2StatNames::PacketsLost, 0);
-		Add(PixelStreaming2StatNames::Jitter, 0);
-		Add(PixelStreaming2StatNames::RoundTripTime, 0);
+		Stats.Add(PixelStreaming2StatNames::FirCount, FStat({ .Name = PixelStreaming2StatNames::FirCount }, 0.f));
+		Stats.Add(PixelStreaming2StatNames::PliCount, FStat({ .Name = PixelStreaming2StatNames::PliCount }, 0.f));
+		Stats.Add(PixelStreaming2StatNames::NackCount, FStat({ .Name = PixelStreaming2StatNames::NackCount }, 0.f));
+		Stats.Add(PixelStreaming2StatNames::RetransmittedBytesSent, FStat({ .Name = PixelStreaming2StatNames::RetransmittedBytesSent }, 0.f));
+		Stats.Add(PixelStreaming2StatNames::TotalEncodeBytesTarget, FStat({ .Name = PixelStreaming2StatNames::TotalEncodeBytesTarget }, 0.f));
+		Stats.Add(PixelStreaming2StatNames::KeyFramesEncoded, FStat({ .Name = PixelStreaming2StatNames::KeyFramesEncoded }, 0.f));
+		Stats.Add(PixelStreaming2StatNames::FrameWidth, FStat({ .Name = PixelStreaming2StatNames::FrameWidth }, 0.f));
+		Stats.Add(PixelStreaming2StatNames::FrameHeight, FStat({ .Name = PixelStreaming2StatNames::FrameHeight }, 0.f));
+		Stats.Add(PixelStreaming2StatNames::HugeFramesSent, FStat({ .Name = PixelStreaming2StatNames::HugeFramesSent }, 0.f));
+		Stats.Add(PixelStreaming2StatNames::PacketsLost, FStat({ .Name = PixelStreaming2StatNames::PacketsLost }, 0.f));
+		Stats.Add(PixelStreaming2StatNames::Jitter, FStat({ .Name = PixelStreaming2StatNames::Jitter }, 0.f));
+		Stats.Add(PixelStreaming2StatNames::RoundTripTime, FStat({ .Name = PixelStreaming2StatNames::RoundTripTime }, 0.f));
+		Stats.Add(PixelStreaming2StatNames::EncoderImplementation, FStat({ .Name = PixelStreaming2StatNames::EncoderImplementation }, FString(TEXT(""))));
 
 		// These are values used to calculate extra values (stores time deltas etc)
-		AddNonRendered(PixelStreaming2StatNames::TargetBitrate);
-		AddNonRendered(PixelStreaming2StatNames::FramesSent);
-		AddNonRendered(PixelStreaming2StatNames::FramesReceived);
-		AddNonRendered(PixelStreaming2StatNames::BytesSent);
-		AddNonRendered(PixelStreaming2StatNames::BytesReceived);
-		AddNonRendered(PixelStreaming2StatNames::QPSum);
-		AddNonRendered(PixelStreaming2StatNames::TotalEncodeTime);
-		AddNonRendered(PixelStreaming2StatNames::FramesEncoded);
-		AddNonRendered(PixelStreaming2StatNames::FramesDecoded);
-		AddNonRendered(PixelStreaming2StatNames::TotalPacketSendDelay);
+		Stats.Add(PixelStreaming2StatNames::TargetBitrate, FStat({ .Name = PixelStreaming2StatNames::TargetBitrate, .DisplayFlags = EDisplayFlags::HIDDEN }, 0.f));
+		Stats.Add(PixelStreaming2StatNames::FramesSent, FStat({ .Name = PixelStreaming2StatNames::FramesSent, .DisplayFlags = EDisplayFlags::HIDDEN }, 0.f));
+		Stats.Add(PixelStreaming2StatNames::FramesReceived, FStat({ .Name = PixelStreaming2StatNames::FramesReceived, .DisplayFlags = EDisplayFlags::HIDDEN }, 0.f));
+		Stats.Add(PixelStreaming2StatNames::BytesSent, FStat({ .Name = PixelStreaming2StatNames::BytesSent, .DisplayFlags = EDisplayFlags::HIDDEN }, 0.f));
+		Stats.Add(PixelStreaming2StatNames::BytesReceived, FStat({ .Name = PixelStreaming2StatNames::BytesReceived, .DisplayFlags = EDisplayFlags::HIDDEN }, 0.f));
+		Stats.Add(PixelStreaming2StatNames::QPSum, FStat({ .Name = PixelStreaming2StatNames::QPSum, .DisplayFlags = EDisplayFlags::HIDDEN }, 0.f));
+		Stats.Add(PixelStreaming2StatNames::TotalEncodeTime, FStat({ .Name = PixelStreaming2StatNames::TotalEncodeTime, .DisplayFlags = EDisplayFlags::HIDDEN }, 0.f));
+		Stats.Add(PixelStreaming2StatNames::FramesEncoded, FStat({ .Name = PixelStreaming2StatNames::FramesEncoded, .DisplayFlags = EDisplayFlags::HIDDEN }, 0.f));
+		Stats.Add(PixelStreaming2StatNames::FramesDecoded, FStat({ .Name = PixelStreaming2StatNames::FramesDecoded, .DisplayFlags = EDisplayFlags::HIDDEN }, 0.f));
+		Stats.Add(PixelStreaming2StatNames::TotalPacketSendDelay, FStat({ .Name = PixelStreaming2StatNames::TotalPacketSendDelay, .DisplayFlags = EDisplayFlags::HIDDEN }, 0.f));
+		Stats.Add(PixelStreaming2StatNames::PacketsSent, FStat({ .Name = PixelStreaming2StatNames::PacketsSent, .DisplayFlags = EDisplayFlags::HIDDEN }, 0.f));
 
 		// Calculated stats below:
 		// FrameSent Per Second
-		AddStatCalculator([](FStatsSink& StatSource, double Period) -> TOptional<FStatData> {
-			FRTCTrackedStat* FramesSentStat = StatSource.Get(PixelStreaming2StatNames::FramesSent);
-			if (FramesSentStat && FramesSentStat->GetLatestStat().StatValue > 0)
+		Calculators.Add([](FStatsSink& StatSource, double Period) -> TOptional<FStat> {
+			FStat* FramesSentStat = StatSource.Get(PixelStreaming2StatNames::FramesSent);
+			if (FramesSentStat && FramesSentStat->GetValue<double>() > 0)
 			{
-				const double FramesSentPerSecond = FramesSentStat->CalculateDelta(Period);
-				FStatData	 FpsStat = FStatData(PixelStreaming2StatNames::FramesSentPerSecond, FramesSentPerSecond, 0);
-				FpsStat.DisplayFlags = FStatData::EDisplayFlags::TEXT | FStatData::EDisplayFlags::GRAPH;
-				return FpsStat;
+				const double FramesSentPerSecond = (FramesSentStat->GetValue<double>() - FramesSentStat->GetPrevValue<double>()) * Period;
+				return FStat({ .Name = PixelStreaming2StatNames::FramesSentPerSecond, .DisplayFlags = static_cast<EDisplayFlags>(EDisplayFlags::TEXT | EDisplayFlags::GRAPH) }, FramesSentPerSecond);
 			}
 			return {};
 		});
 
 		// FramesReceived Per Second
-		AddStatCalculator([](FStatsSink& StatSource, double Period) -> TOptional<FStatData> {
-			FRTCTrackedStat* FramesReceivedStat = StatSource.Get(PixelStreaming2StatNames::FramesReceived);
-			if (FramesReceivedStat && FramesReceivedStat->GetLatestStat().StatValue > 0)
+		Calculators.Add([](FStatsSink& StatSource, double Period) -> TOptional<FStat> {
+			FStat* FramesReceivedStat = StatSource.Get(PixelStreaming2StatNames::FramesReceived);
+			if (FramesReceivedStat && FramesReceivedStat->GetValue<double>() > 0)
 			{
-				const double FramesReceivedPerSecond = FramesReceivedStat->CalculateDelta(Period);
-				return FStatData(PixelStreaming2StatNames::FramesReceivedPerSecond, FramesReceivedPerSecond, 0);
+				const double FramesReceivedPerSecond = (FramesReceivedStat->GetValue<double>() - FramesReceivedStat->GetPrevValue<double>()) * Period;
+				return FStat({
+								 .Name = PixelStreaming2StatNames::FramesReceivedPerSecond,
+							 },
+					FramesReceivedPerSecond);
 			}
 			return {};
 		});
 
 		// Megabits sent Per Second
-		AddStatCalculator([](FStatsSink& StatSource, double Period) -> TOptional<FStatData> {
-			FRTCTrackedStat* BytesSentStat = StatSource.Get(PixelStreaming2StatNames::BytesSent);
-			if (BytesSentStat && BytesSentStat->GetLatestStat().StatValue > 0)
+		Calculators.Add([](FStatsSink& StatSource, double Period) -> TOptional<FStat> {
+			FStat* BytesSentStat = StatSource.Get(PixelStreaming2StatNames::BytesSent);
+			if (BytesSentStat && BytesSentStat->GetValue<double>() > 0)
 			{
-				const double BytesSentPerSecond = BytesSentStat->CalculateDelta(Period);
+				const double BytesSentPerSecond = (BytesSentStat->GetValue<double>() - BytesSentStat->GetPrevValue<double>()) * Period;
 				const double MegabitsPerSecond = BytesSentPerSecond / 1'000'000.0 * 8.0;
-				return FStatData(PixelStreaming2StatNames::BitrateMegabits, MegabitsPerSecond, 2);
+				return FStat({
+								 .Name = PixelStreaming2StatNames::BitrateMegabits,
+							 },
+					MegabitsPerSecond, 2);
 			}
 			return {};
 		});
 
 		// Bits sent Per Second
-		AddStatCalculator([](FStatsSink& StatSource, double Period) -> TOptional<FStatData> {
-			FRTCTrackedStat* BytesSentStat = StatSource.Get(PixelStreaming2StatNames::BytesSent);
-			if (BytesSentStat && BytesSentStat->GetLatestStat().StatValue > 0)
+		Calculators.Add([](FStatsSink& StatSource, double Period) -> TOptional<FStat> {
+			FStat* BytesSentStat = StatSource.Get(PixelStreaming2StatNames::BytesSent);
+			if (BytesSentStat && BytesSentStat->GetValue<double>() > 0)
 			{
-				const double BytesSentPerSecond = BytesSentStat->CalculateDelta(Period);
+				const double BytesSentPerSecond = (BytesSentStat->GetValue<double>() - BytesSentStat->GetPrevValue<double>()) * Period;
 				const double BitsPerSecond = BytesSentPerSecond * 8.0;
-				FStatData	 Stat = FStatData(PixelStreaming2StatNames::Bitrate, BitsPerSecond, 0);
-				Stat.DisplayFlags = FStatData::EDisplayFlags::HIDDEN; // We don't want to display bits per second (too many digits)
-				return Stat;
+				return FStat({ .Name = PixelStreaming2StatNames::Bitrate, .DisplayFlags = EDisplayFlags::HIDDEN }, BitsPerSecond);
 			}
 			return {};
 		});
 
 		// Target megabits sent Per Second
-		AddStatCalculator([](FStatsSink& StatSource, double Period) -> TOptional<FStatData> {
-			FRTCTrackedStat* TargetBpsStats = StatSource.Get(PixelStreaming2StatNames::TargetBitrate);
-			if (TargetBpsStats && TargetBpsStats->GetLatestStat().StatValue > 0)
+		Calculators.Add([](FStatsSink& StatSource, double Period) -> TOptional<FStat> {
+			FStat* TargetBpsStats = StatSource.Get(PixelStreaming2StatNames::TargetBitrate);
+			if (TargetBpsStats && TargetBpsStats->GetValue<double>() > 0)
 			{
-				const double TargetBps = TargetBpsStats->Average();
+				const double TargetBps = (TargetBpsStats->GetValue<double>() + TargetBpsStats->GetPrevValue<double>()) * 0.5f;
 				const double MegabitsPerSecond = TargetBps / 1'000'000.0;
-				return FStatData(PixelStreaming2StatNames::TargetBitrateMegabits, MegabitsPerSecond, 2);
+				return FStat({ .Name = PixelStreaming2StatNames::TargetBitrateMegabits }, MegabitsPerSecond, 2);
 			}
 			return {};
 		});
 
 		// Megabits received Per Second
-		AddStatCalculator([](FStatsSink& StatSource, double Period) -> TOptional<FStatData> {
-			FRTCTrackedStat* BytesReceivedStat = StatSource.Get(PixelStreaming2StatNames::BytesReceived);
-			if (BytesReceivedStat && BytesReceivedStat->GetLatestStat().StatValue > 0)
+		Calculators.Add([](FStatsSink& StatSource, double Period) -> TOptional<FStat> {
+			FStat* BytesReceivedStat = StatSource.Get(PixelStreaming2StatNames::BytesReceived);
+			if (BytesReceivedStat && BytesReceivedStat->GetValue<double>() > 0)
 			{
-				const double BytesReceivedPerSecond = BytesReceivedStat->CalculateDelta(Period);
+				const double BytesReceivedPerSecond = (BytesReceivedStat->GetValue<double>() - BytesReceivedStat->GetPrevValue<double>()) * Period;
 				const double MegabitsPerSecond = BytesReceivedPerSecond / 1000.0 * 8.0;
-				return FStatData(PixelStreaming2StatNames::Bitrate, MegabitsPerSecond, 2);
+				return FStat({ .Name = PixelStreaming2StatNames::Bitrate }, MegabitsPerSecond, 2);
 			}
 			return {};
 		});
 
 		// Encoded fps
-		AddStatCalculator([](FStatsSink& StatSource, double Period) -> TOptional<FStatData> {
-			FRTCTrackedStat* EncodedFramesStat = StatSource.Get(PixelStreaming2StatNames::FramesEncoded);
-			if (EncodedFramesStat && EncodedFramesStat->GetLatestStat().StatValue > 0)
+		Calculators.Add([](FStatsSink& StatSource, double Period) -> TOptional<FStat> {
+			FStat* EncodedFramesStat = StatSource.Get(PixelStreaming2StatNames::FramesEncoded);
+			if (EncodedFramesStat && EncodedFramesStat->GetValue<double>() > 0)
 			{
-				const double EncodedFramesPerSecond = EncodedFramesStat->CalculateDelta(Period);
-				return FStatData(PixelStreaming2StatNames::EncodedFramesPerSecond, EncodedFramesPerSecond, 0);
+				const double EncodedFramesPerSecond = (EncodedFramesStat->GetValue<double>() - EncodedFramesStat->GetPrevValue<double>()) * Period;
+				return FStat({ .Name = PixelStreaming2StatNames::EncodedFramesPerSecond }, EncodedFramesPerSecond);
 			}
 			return {};
 		});
 
 		// Decoded fps
-		AddStatCalculator([](FStatsSink& StatSource, double Period) -> TOptional<FStatData> {
-			FRTCTrackedStat* DecodedFramesStat = StatSource.Get(PixelStreaming2StatNames::FramesDecoded);
-			if (DecodedFramesStat && DecodedFramesStat->GetLatestStat().StatValue > 0)
+		Calculators.Add([](FStatsSink& StatSource, double Period) -> TOptional<FStat> {
+			FStat* DecodedFramesStat = StatSource.Get(PixelStreaming2StatNames::FramesDecoded);
+			if (DecodedFramesStat && DecodedFramesStat->GetValue<double>() > 0)
 			{
-				const double DecodedFramesPerSecond = DecodedFramesStat->CalculateDelta(Period);
-				return FStatData(PixelStreaming2StatNames::DecodedFramesPerSecond, DecodedFramesPerSecond, 0);
+				const double DecodedFramesPerSecond = (DecodedFramesStat->GetValue<double>() - DecodedFramesStat->GetPrevValue<double>()) * Period;
+				return FStat({ .Name = PixelStreaming2StatNames::DecodedFramesPerSecond }, DecodedFramesPerSecond);
 			}
 			return {};
 		});
 
 		// Avg QP Per Second
-		AddStatCalculator([](FStatsSink& StatSource, double Period) -> TOptional<FStatData> {
-			FRTCTrackedStat* QPSumStat = StatSource.Get(PixelStreaming2StatNames::QPSum);
-			FStatData*		 EncodedFramesPerSecond = StatSource.GetCalculatedStat(PixelStreaming2StatNames::EncodedFramesPerSecond);
-			if (QPSumStat && QPSumStat->GetLatestStat().StatValue > 0
-				&& EncodedFramesPerSecond && EncodedFramesPerSecond->StatValue > 0.0)
+		Calculators.Add([](FStatsSink& StatSource, double Period) -> TOptional<FStat> {
+			FStat* QPSumStat = StatSource.Get(PixelStreaming2StatNames::QPSum);
+			FStat* EncodedFramesPerSecond = StatSource.Get(PixelStreaming2StatNames::EncodedFramesPerSecond);
+			if (QPSumStat && QPSumStat->GetValue<double>() > 0
+				&& EncodedFramesPerSecond && EncodedFramesPerSecond->GetValue<double>() > 0.0)
 			{
-				const double QPSumDeltaPerSecond = QPSumStat->CalculateDelta(Period);
-				const double MeanQPPerFrame = QPSumDeltaPerSecond / EncodedFramesPerSecond->StatValue;
-				FName		 StatName = PixelStreaming2StatNames::MeanQPPerSecond;
-				return FStatData(StatName, MeanQPPerFrame, 0);
+				const double QPSumDeltaPerSecond = (QPSumStat->GetValue<double>() - QPSumStat->GetPrevValue<double>()) * Period;
+				const double MeanQPPerFrame = QPSumDeltaPerSecond / EncodedFramesPerSecond->GetValue<double>();
+				return FStat({ .Name = PixelStreaming2StatNames::MeanQPPerSecond }, MeanQPPerFrame);
 			}
 			return {};
 		});
 
 		// Mean EncodeTime (ms) Per Frame
-		AddStatCalculator([](FStatsSink& StatSource, double Period) -> TOptional<FStatData> {
-			FRTCTrackedStat* TotalEncodeTimeStat = StatSource.Get(PixelStreaming2StatNames::TotalEncodeTime);
-			FStatData*		 EncodedFramesPerSecond = StatSource.GetCalculatedStat(PixelStreaming2StatNames::EncodedFramesPerSecond);
-			if (TotalEncodeTimeStat && TotalEncodeTimeStat->GetLatestStat().StatValue > 0
-				&& EncodedFramesPerSecond && EncodedFramesPerSecond->StatValue > 0.0)
+		Calculators.Add([](FStatsSink& StatSource, double Period) -> TOptional<FStat> {
+			FStat* TotalEncodeTimeStat = StatSource.Get(PixelStreaming2StatNames::TotalEncodeTime);
+			FStat* EncodedFramesPerSecond = StatSource.Get(PixelStreaming2StatNames::EncodedFramesPerSecond);
+			if (TotalEncodeTimeStat && TotalEncodeTimeStat->GetValue<double>() > 0
+				&& EncodedFramesPerSecond && EncodedFramesPerSecond->GetValue<double>() > 0.0)
 			{
-				const double TotalEncodeTimePerSecond = TotalEncodeTimeStat->CalculateDelta(Period);
-				const double MeanEncodeTimePerFrameMs = TotalEncodeTimePerSecond / EncodedFramesPerSecond->StatValue * 1000.0;
-				return FStatData(PixelStreaming2StatNames::MeanEncodeTime, MeanEncodeTimePerFrameMs, 2);
+				const double TotalEncodeTimePerSecond = (TotalEncodeTimeStat->GetValue<double>() - TotalEncodeTimeStat->GetPrevValue<double>()) * Period;
+				const double MeanEncodeTimePerFrameMs = TotalEncodeTimePerSecond / EncodedFramesPerSecond->GetValue<double>() * 1000.0;
+				return FStat({ .Name = PixelStreaming2StatNames::MeanEncodeTime }, MeanEncodeTimePerFrameMs, 2);
 			}
 			return {};
 		});
 
 		// Mean SendDelay (ms) Per Frame
-		AddStatCalculator([](FStatsSink& StatSource, double Period) -> TOptional<FStatData> {
-			FRTCTrackedStat* TotalSendDelayStat = StatSource.Get(PixelStreaming2StatNames::TotalPacketSendDelay);
-			FStatData*		 FramesSentPerSecond = StatSource.GetCalculatedStat(PixelStreaming2StatNames::FramesSentPerSecond);
-			if (TotalSendDelayStat && TotalSendDelayStat->GetLatestStat().StatValue > 0
-				&& FramesSentPerSecond && FramesSentPerSecond->StatValue > 0.0)
+		Calculators.Add([](FStatsSink& StatSource, double Period) -> TOptional<FStat> {
+			FStat* TotalSendDelayStat = StatSource.Get(PixelStreaming2StatNames::TotalPacketSendDelay);
+			FStat* TotalPacketsSent = StatSource.Get(PixelStreaming2StatNames::PacketsSent);
+			if (TotalSendDelayStat && TotalSendDelayStat->GetValue<double>() > 0.0
+				&& TotalPacketsSent && TotalPacketsSent->GetValue<double>() > 0.0)
 			{
-				const double TotalSendDelayPerSecond = TotalSendDelayStat->CalculateDelta(Period);
-				const double MeanSendDelayPerFrameMs = TotalSendDelayPerSecond / FramesSentPerSecond->StatValue * 1000.0;
-				return FStatData(PixelStreaming2StatNames::MeanSendDelay, MeanSendDelayPerFrameMs, 2);
+				const double MeanSendDelayPerFrameMs = (TotalSendDelayStat->GetValue<double>() / TotalPacketsSent->GetValue<double>()) * 1000.0;
+				return FStat({ .Name = PixelStreaming2StatNames::MeanSendDelay }, MeanSendDelayPerFrameMs, 2);
 			}
 			return {};
 		});
 
 		// JitterBufferDelay (ms)
-		AddStatCalculator([](FStatsSink& StatSource, double Period) -> TOptional<FStatData> {
-			FRTCTrackedStat* JitterBufferDelayStat = StatSource.Get(PixelStreaming2StatNames::JitterBufferDelay);
-			FStatData*		 FramesReceivedPerSecond = StatSource.GetCalculatedStat(PixelStreaming2StatNames::FramesReceivedPerSecond);
-			if (JitterBufferDelayStat && JitterBufferDelayStat->GetLatestStat().StatValue > 0
-				&& FramesReceivedPerSecond && FramesReceivedPerSecond->StatValue > 0.0)
+		Calculators.Add([](FStatsSink& StatSource, double Period) -> TOptional<FStat> {
+			FStat* JitterBufferDelayStat = StatSource.Get(PixelStreaming2StatNames::JitterBufferDelay);
+			FStat* FramesReceivedPerSecond = StatSource.Get(PixelStreaming2StatNames::FramesReceivedPerSecond);
+			if (JitterBufferDelayStat && JitterBufferDelayStat->GetValue<double>() > 0
+				&& FramesReceivedPerSecond && FramesReceivedPerSecond->GetValue<double>() > 0.0)
 			{
-				const double TotalJitterBufferDelayPerSecond = JitterBufferDelayStat->CalculateDelta(Period);
-				const double MeanJitterBufferDelayMs = TotalJitterBufferDelayPerSecond / FramesReceivedPerSecond->StatValue * 1000.0;
-				return FStatData(PixelStreaming2StatNames::JitterBufferDelay, MeanJitterBufferDelayMs, 2);
+				const double TotalJitterBufferDelayPerSecond = (JitterBufferDelayStat->GetValue<double>() - JitterBufferDelayStat->GetPrevValue<double>()) * Period;
+				const double MeanJitterBufferDelayMs = TotalJitterBufferDelayPerSecond / FramesReceivedPerSecond->GetValue<double>() * 1000.0;
+				return FStat({ .Name = PixelStreaming2StatNames::JitterBufferDelay }, MeanJitterBufferDelayMs, 2);
 			}
 			return {};
 		});
@@ -393,56 +719,56 @@ namespace UE::PixelStreaming2
 			return;
 		}
 
-		for (TPair<FName, FRTCTrackedStat>& Tuple : Stats)
+		for (TPair<FName, FStat>& Tuple : Stats)
 		{
-			double NewValue = 0;
+			FStatVariant NewValue;
 			if (Tuple.Key == PixelStreaming2StatNames::FirCount)
 			{
-				NewValue = InStats._local._firCount;
+				NewValue = FStatVariant(TInPlaceType<double>(), InStats._local._firCount);
 			}
 			else if (Tuple.Key == PixelStreaming2StatNames::PliCount)
 			{
-				NewValue = InStats._local._pliCount;
+				NewValue = FStatVariant(TInPlaceType<double>(), InStats._local._pliCount);
 			}
 			else if (Tuple.Key == PixelStreaming2StatNames::NackCount)
 			{
-				NewValue = InStats._local._nackCount;
+				NewValue = FStatVariant(TInPlaceType<double>(), InStats._local._nackCount);
 			}
 			else if (Tuple.Key == PixelStreaming2StatNames::RetransmittedBytesSent)
 			{
-				NewValue = InStats._local._retransmittedBytesSent;
+				NewValue = FStatVariant(TInPlaceType<double>(), InStats._local._retransmittedBytesSent);
 			}
 			else if (Tuple.Key == PixelStreaming2StatNames::TotalEncodeBytesTarget)
 			{
-				NewValue = InStats._local._totalEncodedBytesTarget;
+				NewValue = FStatVariant(TInPlaceType<double>(), InStats._local._totalEncodedBytesTarget);
 			}
 			else if (Tuple.Key == PixelStreaming2StatNames::KeyFramesEncoded)
 			{
-				NewValue = InStats._local._keyFramesEncoded;
+				NewValue = FStatVariant(TInPlaceType<double>(), InStats._local._keyFramesEncoded);
 			}
 			else if (Tuple.Key == PixelStreaming2StatNames::FrameWidth)
 			{
-				NewValue = InStats._local._frameWidth;
+				NewValue = FStatVariant(TInPlaceType<double>(), InStats._local._frameWidth);
 			}
 			else if (Tuple.Key == PixelStreaming2StatNames::FrameHeight)
 			{
-				NewValue = InStats._local._frameHeight;
+				NewValue = FStatVariant(TInPlaceType<double>(), InStats._local._frameHeight);
 			}
 			else if (Tuple.Key == PixelStreaming2StatNames::HugeFramesSent)
 			{
-				NewValue = InStats._local._hugeFramesSent;
+				NewValue = FStatVariant(TInPlaceType<double>(), InStats._local._hugeFramesSent);
 			}
 			else if (Tuple.Key == PixelStreaming2StatNames::TotalPacketSendDelay)
 			{
-				NewValue = InStats._local._totalPacketSendDelay;
+				NewValue = FStatVariant(TInPlaceType<double>(), InStats._local._totalPacketSendDelay);
 			}
 			else if (Tuple.Key == PixelStreaming2StatNames::TargetBitrate)
 			{
-				NewValue = InStats._local._targetBitrate;
+				NewValue = FStatVariant(TInPlaceType<double>(), InStats._local._targetBitrate);
 			}
 			else if (Tuple.Key == PixelStreaming2StatNames::FramesSent)
 			{
-				NewValue = InStats._local._framesSent;
+				NewValue = FStatVariant(TInPlaceType<double>(), InStats._local._framesSent);
 			}
 			else if (Tuple.Key == PixelStreaming2StatNames::FramesReceived)
 			{
@@ -450,7 +776,7 @@ namespace UE::PixelStreaming2
 			}
 			else if (Tuple.Key == PixelStreaming2StatNames::BytesSent)
 			{
-				NewValue = InStats._local._bytesSent;
+				NewValue = FStatVariant(TInPlaceType<double>(), InStats._local._bytesSent);
 			}
 			else if (Tuple.Key == PixelStreaming2StatNames::BytesReceived)
 			{
@@ -458,36 +784,49 @@ namespace UE::PixelStreaming2
 			}
 			else if (Tuple.Key == PixelStreaming2StatNames::QPSum)
 			{
-				NewValue = InStats._local._qpSum;
+				NewValue = FStatVariant(TInPlaceType<double>(), InStats._local._qpSum);
 			}
 			else if (Tuple.Key == PixelStreaming2StatNames::TotalEncodeTime)
 			{
-				NewValue = InStats._local._totalEncodeTime;
+				NewValue = FStatVariant(TInPlaceType<double>(), InStats._local._totalEncodeTime);
 			}
 			else if (Tuple.Key == PixelStreaming2StatNames::FramesEncoded)
 			{
-				NewValue = InStats._local._framesEncoded;
+				NewValue = FStatVariant(TInPlaceType<double>(), InStats._local._framesEncoded);
 			}
 			else if (Tuple.Key == PixelStreaming2StatNames::FramesDecoded)
 			{
 				// TODO(Nazar.Rudenko): Available for inbound tracks only
 			}
+			else if (Tuple.Key == PixelStreaming2StatNames::EncoderImplementation)
+			{
+				NewValue = FStatVariant(TInPlaceType<FString>(), ToString(InStats._local._encoderImplementation));
+			}
+			else if (Tuple.Key == PixelStreaming2StatNames::PacketsSent)
+			{
+				NewValue = FStatVariant(TInPlaceType<double>(), InStats._local._packetsSent);
+			}
 			else if (Tuple.Key == PixelStreaming2StatNames::PacketsLost)
 			{
-				NewValue = InStats._remote._packetsLost;
+				NewValue = FStatVariant(TInPlaceType<double>(), InStats._remote._packetsLost);
 			}
 			else if (Tuple.Key == PixelStreaming2StatNames::Jitter)
 			{
-				NewValue = InStats._remote._jitter;
+				NewValue = FStatVariant(TInPlaceType<double>(), InStats._remote._jitter);
 			}
 			else if (Tuple.Key == PixelStreaming2StatNames::RoundTripTime)
 			{
-				NewValue = InStats._remote._roundTripTime;
+				NewValue = FStatVariant(TInPlaceType<double>(), InStats._remote._roundTripTime);
+			}
+
+			if (NewValue.GetIndex() == FStatVariant::IndexOfType<TYPE_OF_NULLPTR>())
+			{
+				continue;
 			}
 
-			if (UpdateValue(NewValue, &Tuple.Value))
+			if (Tuple.Value.SetValue(NewValue))
 			{
-				PSStats->StorePeerStat(PeerId, Category, Tuple.Value.GetLatestStat());
+				PSStats->StorePeerStat(PeerId, Category, Tuple.Value);
 			}
 		}
 
@@ -501,181 +840,175 @@ namespace UE::PixelStreaming2
 		: FStatsSink(InCategory)
 	{
 		// These stats will be extracted from the stat reports and emitted straight to screen
-		Add(PixelStreaming2StatNames::FirCount, 0);
-		Add(PixelStreaming2StatNames::PliCount, 0);
-		Add(PixelStreaming2StatNames::NackCount, 0);
-		Add(PixelStreaming2StatNames::RetransmittedBytesSent, 0);
-		Add(PixelStreaming2StatNames::TotalEncodeBytesTarget, 0);
-		Add(PixelStreaming2StatNames::KeyFramesEncoded, 0);
-		Add(PixelStreaming2StatNames::FrameWidth, 0);
-		Add(PixelStreaming2StatNames::FrameHeight, 0);
-		Add(PixelStreaming2StatNames::HugeFramesSent, 0);
-		Add(PixelStreaming2StatNames::PacketsLost, 0);
-		Add(PixelStreaming2StatNames::Jitter, 0);
-		Add(PixelStreaming2StatNames::RoundTripTime, 0);
+		Stats.Add(PixelStreaming2StatNames::FirCount, FStat({ .Name = PixelStreaming2StatNames::FirCount }, 0.f));
+		Stats.Add(PixelStreaming2StatNames::PliCount, FStat({ .Name = PixelStreaming2StatNames::PliCount }, 0.f));
+		Stats.Add(PixelStreaming2StatNames::NackCount, FStat({ .Name = PixelStreaming2StatNames::NackCount }, 0.f));
+		Stats.Add(PixelStreaming2StatNames::RetransmittedBytesSent, FStat({ .Name = PixelStreaming2StatNames::RetransmittedBytesSent }, 0.f));
+		Stats.Add(PixelStreaming2StatNames::TotalEncodeBytesTarget, FStat({ .Name = PixelStreaming2StatNames::TotalEncodeBytesTarget }, 0.f));
+		Stats.Add(PixelStreaming2StatNames::KeyFramesEncoded, FStat({ .Name = PixelStreaming2StatNames::KeyFramesEncoded }, 0.f));
+		Stats.Add(PixelStreaming2StatNames::FrameWidth, FStat({ .Name = PixelStreaming2StatNames::FrameWidth }, 0.f));
+		Stats.Add(PixelStreaming2StatNames::FrameHeight, FStat({ .Name = PixelStreaming2StatNames::FrameHeight }, 0.f));
+		Stats.Add(PixelStreaming2StatNames::HugeFramesSent, FStat({ .Name = PixelStreaming2StatNames::HugeFramesSent }, 0.f));
+		Stats.Add(PixelStreaming2StatNames::PacketsLost, FStat({ .Name = PixelStreaming2StatNames::PacketsLost }, 0.f));
+		Stats.Add(PixelStreaming2StatNames::Jitter, FStat({ .Name = PixelStreaming2StatNames::Jitter }, 0.f));
+		Stats.Add(PixelStreaming2StatNames::RoundTripTime, FStat({ .Name = PixelStreaming2StatNames::RoundTripTime }, 0.f));
 
 		// These are values used to calculate extra values (stores time deltas etc)
-		AddNonRendered(PixelStreaming2StatNames::TargetBitrate);
-		AddNonRendered(PixelStreaming2StatNames::FramesSent);
-		AddNonRendered(PixelStreaming2StatNames::FramesReceived);
-		AddNonRendered(PixelStreaming2StatNames::BytesSent);
-		AddNonRendered(PixelStreaming2StatNames::BytesReceived);
-		AddNonRendered(PixelStreaming2StatNames::QPSum);
-		AddNonRendered(PixelStreaming2StatNames::TotalEncodeTime);
-		AddNonRendered(PixelStreaming2StatNames::FramesEncoded);
-		AddNonRendered(PixelStreaming2StatNames::FramesDecoded);
-		AddNonRendered(PixelStreaming2StatNames::TotalPacketSendDelay);
+		Stats.Add(PixelStreaming2StatNames::TargetBitrate, FStat({ .Name = PixelStreaming2StatNames::TargetBitrate, .DisplayFlags = EDisplayFlags::HIDDEN }, 0.f));
+		Stats.Add(PixelStreaming2StatNames::FramesSent, FStat({ .Name = PixelStreaming2StatNames::FramesSent, .DisplayFlags = EDisplayFlags::HIDDEN }, 0.f));
+		Stats.Add(PixelStreaming2StatNames::FramesReceived, FStat({ .Name = PixelStreaming2StatNames::FramesReceived, .DisplayFlags = EDisplayFlags::HIDDEN }, 0.f));
+		Stats.Add(PixelStreaming2StatNames::BytesSent, FStat({ .Name = PixelStreaming2StatNames::BytesSent, .DisplayFlags = EDisplayFlags::HIDDEN }, 0.f));
+		Stats.Add(PixelStreaming2StatNames::BytesReceived, FStat({ .Name = PixelStreaming2StatNames::BytesReceived, .DisplayFlags = EDisplayFlags::HIDDEN }, 0.f));
+		Stats.Add(PixelStreaming2StatNames::QPSum, FStat({ .Name = PixelStreaming2StatNames::QPSum, .DisplayFlags = EDisplayFlags::HIDDEN }, 0.f));
+		Stats.Add(PixelStreaming2StatNames::TotalEncodeTime, FStat({ .Name = PixelStreaming2StatNames::TotalEncodeTime, .DisplayFlags = EDisplayFlags::HIDDEN }, 0.f));
+		Stats.Add(PixelStreaming2StatNames::FramesEncoded, FStat({ .Name = PixelStreaming2StatNames::FramesEncoded, .DisplayFlags = EDisplayFlags::HIDDEN }, 0.f));
+		Stats.Add(PixelStreaming2StatNames::FramesDecoded, FStat({ .Name = PixelStreaming2StatNames::FramesDecoded, .DisplayFlags = EDisplayFlags::HIDDEN }, 0.f));
+		Stats.Add(PixelStreaming2StatNames::TotalPacketSendDelay, FStat({ .Name = PixelStreaming2StatNames::TotalPacketSendDelay, .DisplayFlags = EDisplayFlags::HIDDEN }, 0.f));
 
 		// Calculated stats below:
 		// FrameSent Per Second
-		AddStatCalculator([](FStatsSink& StatSource, double Period) -> TOptional<FStatData> {
-			FRTCTrackedStat* FramesSentStat = StatSource.Get(PixelStreaming2StatNames::FramesSent);
-			if (FramesSentStat && FramesSentStat->GetLatestStat().StatValue > 0)
+		Calculators.Add([](FStatsSink& StatSource, double Period) -> TOptional<FStat> {
+			FStat* FramesSentStat = StatSource.Get(PixelStreaming2StatNames::FramesSent);
+			if (FramesSentStat && FramesSentStat->GetValue<double>() > 0)
 			{
-				const double FramesSentPerSecond = FramesSentStat->CalculateDelta(Period);
-				FStatData	 FpsStat = FStatData(PixelStreaming2StatNames::FramesSentPerSecond, FramesSentPerSecond, 0);
-				FpsStat.DisplayFlags = FStatData::EDisplayFlags::TEXT | FStatData::EDisplayFlags::GRAPH;
-				return FpsStat;
+				const double FramesSentPerSecond = (FramesSentStat->GetValue<double>() - FramesSentStat->GetPrevValue<double>()) * Period;
+				return FStat({ .Name = PixelStreaming2StatNames::FramesSentPerSecond, .DisplayFlags = static_cast<EDisplayFlags>(EDisplayFlags::TEXT | EDisplayFlags::GRAPH) }, FramesSentPerSecond);
 			}
 			return {};
 		});
 
 		// FramesReceived Per Second
-		AddStatCalculator([](FStatsSink& StatSource, double Period) -> TOptional<FStatData> {
-			FRTCTrackedStat* FramesReceivedStat = StatSource.Get(PixelStreaming2StatNames::FramesReceived);
-			if (FramesReceivedStat && FramesReceivedStat->GetLatestStat().StatValue > 0)
+		Calculators.Add([](FStatsSink& StatSource, double Period) -> TOptional<FStat> {
+			FStat* FramesReceivedStat = StatSource.Get(PixelStreaming2StatNames::FramesReceived);
+			if (FramesReceivedStat && FramesReceivedStat->GetValue<double>() > 0)
 			{
-				const double FramesReceivedPerSecond = FramesReceivedStat->CalculateDelta(Period);
-				return FStatData(PixelStreaming2StatNames::FramesReceivedPerSecond, FramesReceivedPerSecond, 0);
+				const double FramesReceivedPerSecond = (FramesReceivedStat->GetValue<double>() - FramesReceivedStat->GetPrevValue<double>()) * Period;
+				return FStat({ .Name = PixelStreaming2StatNames::FramesReceivedPerSecond }, FramesReceivedPerSecond);
 			}
 			return {};
 		});
 
 		// Megabits sent Per Second
-		AddStatCalculator([](FStatsSink& StatSource, double Period) -> TOptional<FStatData> {
-			FRTCTrackedStat* BytesSentStat = StatSource.Get(PixelStreaming2StatNames::BytesSent);
-			if (BytesSentStat && BytesSentStat->GetLatestStat().StatValue > 0)
+		Calculators.Add([](FStatsSink& StatSource, double Period) -> TOptional<FStat> {
+			FStat* BytesSentStat = StatSource.Get(PixelStreaming2StatNames::BytesSent);
+			if (BytesSentStat && BytesSentStat->GetValue<double>() > 0)
 			{
-				const double BytesSentPerSecond = BytesSentStat->CalculateDelta(Period);
+				const double BytesSentPerSecond = (BytesSentStat->GetValue<double>() - BytesSentStat->GetPrevValue<double>()) * Period;
 				const double MegabitsPerSecond = BytesSentPerSecond / 1'000'000.0 * 8.0;
-				return FStatData(PixelStreaming2StatNames::BitrateMegabits, MegabitsPerSecond, 2);
+				return FStat({ .Name = PixelStreaming2StatNames::BitrateMegabits }, MegabitsPerSecond, 2);
 			}
 			return {};
 		});
 
 		// Bits sent Per Second
-		AddStatCalculator([](FStatsSink& StatSource, double Period) -> TOptional<FStatData> {
-			FRTCTrackedStat* BytesSentStat = StatSource.Get(PixelStreaming2StatNames::BytesSent);
-			if (BytesSentStat && BytesSentStat->GetLatestStat().StatValue > 0)
+		Calculators.Add([](FStatsSink& StatSource, double Period) -> TOptional<FStat> {
+			FStat* BytesSentStat = StatSource.Get(PixelStreaming2StatNames::BytesSent);
+			if (BytesSentStat && BytesSentStat->GetValue<double>() > 0)
 			{
-				const double BytesSentPerSecond = BytesSentStat->CalculateDelta(Period);
+				const double BytesSentPerSecond = (BytesSentStat->GetValue<double>() - BytesSentStat->GetPrevValue<double>()) * Period;
 				const double BitsPerSecond = BytesSentPerSecond * 8.0;
-				FStatData	 Stat = FStatData(PixelStreaming2StatNames::Bitrate, BitsPerSecond, 0);
-				Stat.DisplayFlags = FStatData::EDisplayFlags::HIDDEN; // We don't want to display bits per second (too many digits)
-				return Stat;
+				return FStat({ .Name = PixelStreaming2StatNames::Bitrate, .DisplayFlags = EDisplayFlags::HIDDEN }, BitsPerSecond);
 			}
 			return {};
 		});
 
 		// Target megabits sent Per Second
-		AddStatCalculator([](FStatsSink& StatSource, double Period) -> TOptional<FStatData> {
-			FRTCTrackedStat* TargetBpsStats = StatSource.Get(PixelStreaming2StatNames::TargetBitrate);
-			if (TargetBpsStats && TargetBpsStats->GetLatestStat().StatValue > 0)
+		Calculators.Add([](FStatsSink& StatSource, double Period) -> TOptional<FStat> {
+			FStat* TargetBpsStats = StatSource.Get(PixelStreaming2StatNames::TargetBitrate);
+			if (TargetBpsStats && TargetBpsStats->GetValue<double>() > 0)
 			{
-				const double TargetBps = TargetBpsStats->Average();
+				const double TargetBps = (TargetBpsStats->GetValue<double>() + TargetBpsStats->GetPrevValue<double>()) * 0.5f;
 				const double MegabitsPerSecond = TargetBps / 1'000'000.0;
-				return FStatData(PixelStreaming2StatNames::TargetBitrateMegabits, MegabitsPerSecond, 2);
+				return FStat({ .Name = PixelStreaming2StatNames::TargetBitrateMegabits }, MegabitsPerSecond, 2);
 			}
 			return {};
 		});
 
 		// Megabits received Per Second
-		AddStatCalculator([](FStatsSink& StatSource, double Period) -> TOptional<FStatData> {
-			FRTCTrackedStat* BytesReceivedStat = StatSource.Get(PixelStreaming2StatNames::BytesReceived);
-			if (BytesReceivedStat && BytesReceivedStat->GetLatestStat().StatValue > 0)
+		Calculators.Add([](FStatsSink& StatSource, double Period) -> TOptional<FStat> {
+			FStat* BytesReceivedStat = StatSource.Get(PixelStreaming2StatNames::BytesReceived);
+			if (BytesReceivedStat && BytesReceivedStat->GetValue<double>() > 0)
 			{
-				const double BytesReceivedPerSecond = BytesReceivedStat->CalculateDelta(Period);
+				const double BytesReceivedPerSecond = (BytesReceivedStat->GetValue<double>() - BytesReceivedStat->GetPrevValue<double>()) * Period;
 				const double MegabitsPerSecond = BytesReceivedPerSecond / 1000.0 * 8.0;
-				return FStatData(PixelStreaming2StatNames::Bitrate, MegabitsPerSecond, 2);
+				return FStat({ .Name = PixelStreaming2StatNames::Bitrate }, MegabitsPerSecond, 2);
 			}
 			return {};
 		});
 
 		// Encoded fps
-		AddStatCalculator([](FStatsSink& StatSource, double Period) -> TOptional<FStatData> {
-			FRTCTrackedStat* EncodedFramesStat = StatSource.Get(PixelStreaming2StatNames::FramesEncoded);
-			if (EncodedFramesStat && EncodedFramesStat->GetLatestStat().StatValue > 0)
+		Calculators.Add([](FStatsSink& StatSource, double Period) -> TOptional<FStat> {
+			FStat* EncodedFramesStat = StatSource.Get(PixelStreaming2StatNames::FramesEncoded);
+			if (EncodedFramesStat && EncodedFramesStat->GetValue<double>() > 0)
 			{
-				const double EncodedFramesPerSecond = EncodedFramesStat->CalculateDelta(Period);
-				return FStatData(PixelStreaming2StatNames::EncodedFramesPerSecond, EncodedFramesPerSecond, 0);
+				const double EncodedFramesPerSecond = (EncodedFramesStat->GetValue<double>() - EncodedFramesStat->GetPrevValue<double>()) * Period;
+				return FStat({ .Name = PixelStreaming2StatNames::EncodedFramesPerSecond }, EncodedFramesPerSecond);
 			}
 			return {};
 		});
 
 		// Decoded fps
-		AddStatCalculator([](FStatsSink& StatSource, double Period) -> TOptional<FStatData> {
-			FRTCTrackedStat* DecodedFramesStat = StatSource.Get(PixelStreaming2StatNames::FramesDecoded);
-			if (DecodedFramesStat && DecodedFramesStat->GetLatestStat().StatValue > 0)
+		Calculators.Add([](FStatsSink& StatSource, double Period) -> TOptional<FStat> {
+			FStat* DecodedFramesStat = StatSource.Get(PixelStreaming2StatNames::FramesDecoded);
+			if (DecodedFramesStat && DecodedFramesStat->GetValue<double>() > 0)
 			{
-				const double DecodedFramesPerSecond = DecodedFramesStat->CalculateDelta(Period);
-				return FStatData(PixelStreaming2StatNames::DecodedFramesPerSecond, DecodedFramesPerSecond, 0);
+				const double DecodedFramesPerSecond = (DecodedFramesStat->GetValue<double>() - DecodedFramesStat->GetPrevValue<double>()) * Period;
+				return FStat({ .Name = PixelStreaming2StatNames::DecodedFramesPerSecond }, DecodedFramesPerSecond);
 			}
 			return {};
 		});
 
 		// Avg QP Per Second
-		AddStatCalculator([](FStatsSink& StatSource, double Period) -> TOptional<FStatData> {
-			FRTCTrackedStat* QPSumStat = StatSource.Get(PixelStreaming2StatNames::QPSum);
-			FStatData*		 EncodedFramesPerSecond = StatSource.GetCalculatedStat(PixelStreaming2StatNames::EncodedFramesPerSecond);
-			if (QPSumStat && QPSumStat->GetLatestStat().StatValue > 0
-				&& EncodedFramesPerSecond && EncodedFramesPerSecond->StatValue > 0.0)
+		Calculators.Add([](FStatsSink& StatSource, double Period) -> TOptional<FStat> {
+			FStat* QPSumStat = StatSource.Get(PixelStreaming2StatNames::QPSum);
+			FStat* EncodedFramesPerSecond = StatSource.Get(PixelStreaming2StatNames::EncodedFramesPerSecond);
+			if (QPSumStat && QPSumStat->GetValue<double>() > 0
+				&& EncodedFramesPerSecond && EncodedFramesPerSecond->GetValue<double>() > 0.0)
 			{
-				const double QPSumDeltaPerSecond = QPSumStat->CalculateDelta(Period);
-				const double MeanQPPerFrame = QPSumDeltaPerSecond / EncodedFramesPerSecond->StatValue;
-				FName		 StatName = PixelStreaming2StatNames::MeanQPPerSecond;
-				return FStatData(StatName, MeanQPPerFrame, 0);
+				const double QPSumDeltaPerSecond = (QPSumStat->GetValue<double>() - QPSumStat->GetPrevValue<double>()) * Period;
+				const double MeanQPPerFrame = QPSumDeltaPerSecond / EncodedFramesPerSecond->GetValue<double>();
+				return FStat({ .Name = PixelStreaming2StatNames::MeanQPPerSecond }, MeanQPPerFrame);
 			}
 			return {};
 		});
 
 		// Mean EncodeTime (ms) Per Frame
-		AddStatCalculator([](FStatsSink& StatSource, double Period) -> TOptional<FStatData> {
-			FRTCTrackedStat* TotalEncodeTimeStat = StatSource.Get(PixelStreaming2StatNames::TotalEncodeTime);
-			FStatData*		 EncodedFramesPerSecond = StatSource.GetCalculatedStat(PixelStreaming2StatNames::EncodedFramesPerSecond);
-			if (TotalEncodeTimeStat && TotalEncodeTimeStat->GetLatestStat().StatValue > 0
-				&& EncodedFramesPerSecond && EncodedFramesPerSecond->StatValue > 0.0)
+		Calculators.Add([](FStatsSink& StatSource, double Period) -> TOptional<FStat> {
+			FStat* TotalEncodeTimeStat = StatSource.Get(PixelStreaming2StatNames::TotalEncodeTime);
+			FStat* EncodedFramesPerSecond = StatSource.Get(PixelStreaming2StatNames::EncodedFramesPerSecond);
+			if (TotalEncodeTimeStat && TotalEncodeTimeStat->GetValue<double>() > 0
+				&& EncodedFramesPerSecond && EncodedFramesPerSecond->GetValue<double>() > 0.0)
 			{
-				const double TotalEncodeTimePerSecond = TotalEncodeTimeStat->CalculateDelta(Period);
-				const double MeanEncodeTimePerFrameMs = TotalEncodeTimePerSecond / EncodedFramesPerSecond->StatValue * 1000.0;
-				return FStatData(PixelStreaming2StatNames::MeanEncodeTime, MeanEncodeTimePerFrameMs, 2);
+				const double TotalEncodeTimePerSecond = (TotalEncodeTimeStat->GetValue<double>() - TotalEncodeTimeStat->GetPrevValue<double>()) * Period;
+				const double MeanEncodeTimePerFrameMs = TotalEncodeTimePerSecond / EncodedFramesPerSecond->GetValue<double>() * 1000.0;
+				return FStat({ .Name = PixelStreaming2StatNames::MeanEncodeTime }, MeanEncodeTimePerFrameMs, 2);
 			}
 			return {};
 		});
 
 		// Mean SendDelay (ms) Per Frame
-		AddStatCalculator([](FStatsSink& StatSource, double Period) -> TOptional<FStatData> {
-			FRTCTrackedStat* TotalSendDelayStat = StatSource.Get(PixelStreaming2StatNames::TotalPacketSendDelay);
-			FStatData*		 FramesSentPerSecond = StatSource.GetCalculatedStat(PixelStreaming2StatNames::FramesSentPerSecond);
-			if (TotalSendDelayStat && TotalSendDelayStat->GetLatestStat().StatValue > 0
-				&& FramesSentPerSecond && FramesSentPerSecond->StatValue > 0.0)
+		Calculators.Add([](FStatsSink& StatSource, double Period) -> TOptional<FStat> {
+			FStat* TotalSendDelayStat = StatSource.Get(PixelStreaming2StatNames::TotalPacketSendDelay);
+			FStat* TotalPacketsSent = StatSource.Get(PixelStreaming2StatNames::PacketsSent);
+			if (TotalSendDelayStat && TotalSendDelayStat->GetValue<double>() > 0.0
+				&& TotalPacketsSent && TotalPacketsSent->GetValue<double>() > 0.0)
 			{
-				const double TotalSendDelayPerSecond = TotalSendDelayStat->CalculateDelta(Period);
-				const double MeanSendDelayPerFrameMs = TotalSendDelayPerSecond / FramesSentPerSecond->StatValue * 1000.0;
-				return FStatData(PixelStreaming2StatNames::MeanSendDelay, MeanSendDelayPerFrameMs, 2);
+				const double MeanSendDelayPerFrameMs = (TotalSendDelayStat->GetValue<double>() / TotalPacketsSent->GetValue<double>()) * 1000.0;
+				return FStat({ .Name = PixelStreaming2StatNames::MeanSendDelay }, MeanSendDelayPerFrameMs, 2);
 			}
 			return {};
 		});
 
 		// JitterBufferDelay (ms)
-		AddStatCalculator([](FStatsSink& StatSource, double Period) -> TOptional<FStatData> {
-			FRTCTrackedStat* JitterBufferDelayStat = StatSource.Get(PixelStreaming2StatNames::JitterBufferDelay);
-			FStatData*		 FramesReceivedPerSecond = StatSource.GetCalculatedStat(PixelStreaming2StatNames::FramesReceivedPerSecond);
-			if (JitterBufferDelayStat && JitterBufferDelayStat->GetLatestStat().StatValue > 0
-				&& FramesReceivedPerSecond && FramesReceivedPerSecond->StatValue > 0.0)
+		Calculators.Add([](FStatsSink& StatSource, double Period) -> TOptional<FStat> {
+			FStat* JitterBufferDelayStat = StatSource.Get(PixelStreaming2StatNames::JitterBufferDelay);
+			FStat* FramesReceivedPerSecond = StatSource.Get(PixelStreaming2StatNames::FramesReceivedPerSecond);
+			if (JitterBufferDelayStat && JitterBufferDelayStat->GetValue<double>() > 0
+				&& FramesReceivedPerSecond && FramesReceivedPerSecond->GetValue<double>() > 0.0)
 			{
-				const double TotalJitterBufferDelayPerSecond = JitterBufferDelayStat->CalculateDelta(Period);
-				const double MeanJitterBufferDelayMs = TotalJitterBufferDelayPerSecond / FramesReceivedPerSecond->StatValue * 1000.0;
-				return FStatData(PixelStreaming2StatNames::JitterBufferDelay, MeanJitterBufferDelayMs, 2);
+				const double TotalJitterBufferDelayPerSecond = (JitterBufferDelayStat->GetValue<double>() - JitterBufferDelayStat->GetPrevValue<double>()) * Period;
+				const double MeanJitterBufferDelayMs = TotalJitterBufferDelayPerSecond / FramesReceivedPerSecond->GetValue<double>() * 1000.0;
+				return FStat({ .Name = PixelStreaming2StatNames::JitterBufferDelay }, MeanJitterBufferDelayMs, 2);
 			}
 			return {};
 		});
@@ -689,37 +1022,42 @@ namespace UE::PixelStreaming2
 			return;
 		}
 
-		for (TPair<FName, FRTCTrackedStat>& Tuple : Stats)
+		for (TPair<FName, FStat>& Tuple : Stats)
 		{
-			double NewValue = 0;
+			FStatVariant NewValue;
 			if (Tuple.Key == PixelStreaming2StatNames::TotalPacketSendDelay)
 			{
-				NewValue = InStats._local._totalPacketSendDelay;
+				NewValue = FStatVariant(TInPlaceType<double>(), InStats._local._totalPacketSendDelay);
 			}
 			else if (Tuple.Key == PixelStreaming2StatNames::TargetBitrate)
 			{
-				NewValue = InStats._local._targetBitrate;
+				NewValue = FStatVariant(TInPlaceType<double>(), InStats._local._targetBitrate);
 			}
 			else if (Tuple.Key == PixelStreaming2StatNames::BytesSent)
 			{
-				NewValue = InStats._local._bytesSent;
+				NewValue = FStatVariant(TInPlaceType<double>(), InStats._local._bytesSent);
 			}
 			else if (Tuple.Key == PixelStreaming2StatNames::PacketsLost)
 			{
-				NewValue = InStats._remote._packetsLost;
+				NewValue = FStatVariant(TInPlaceType<double>(), InStats._remote._packetsLost);
 			}
 			else if (Tuple.Key == PixelStreaming2StatNames::Jitter)
 			{
-				NewValue = InStats._remote._jitter;
+				NewValue = FStatVariant(TInPlaceType<double>(), InStats._remote._jitter);
 			}
 			else if (Tuple.Key == PixelStreaming2StatNames::RoundTripTime)
 			{
-				NewValue = InStats._remote._roundTripTime;
+				NewValue = FStatVariant(TInPlaceType<double>(), InStats._remote._roundTripTime);
 			}
 
-			if (UpdateValue(NewValue, &Tuple.Value))
+			if (NewValue.GetIndex() == FStatVariant::IndexOfType<TYPE_OF_NULLPTR>())
 			{
-				PSStats->StorePeerStat(PeerId, Category, Tuple.Value.GetLatestStat());
+				continue;
+			}
+
+			if (Tuple.Value.SetValue(NewValue))
+			{
+				PSStats->StorePeerStat(PeerId, Category, Tuple.Value);
 			}
 		}
 
@@ -733,182 +1071,176 @@ namespace UE::PixelStreaming2
 		: FStatsSink(InCategory)
 	{
 		// These stats will be extracted from the stat reports and emitted straight to screen
-		Add(PixelStreaming2StatNames::FirCount, 0);
-		Add(PixelStreaming2StatNames::PliCount, 0);
-		Add(PixelStreaming2StatNames::NackCount, 0);
-		Add(PixelStreaming2StatNames::RetransmittedBytesReceived, 0);
-		Add(PixelStreaming2StatNames::RetransmittedPacketsReceived, 0);
-		Add(PixelStreaming2StatNames::TotalEncodeBytesTarget, 0);
-		Add(PixelStreaming2StatNames::KeyFramesDecoded, 0);
-		Add(PixelStreaming2StatNames::FrameWidth, 0);
-		Add(PixelStreaming2StatNames::FrameHeight, 0);
-		Add(PixelStreaming2StatNames::HugeFramesSent, 0);
-		Add(PixelStreaming2StatNames::PacketsLost, 0);
-		Add(PixelStreaming2StatNames::Jitter, 0);
-		Add(PixelStreaming2StatNames::RoundTripTime, 0);
+		Stats.Add(PixelStreaming2StatNames::FirCount, FStat({ .Name = PixelStreaming2StatNames::FirCount }, 0.f));
+		Stats.Add(PixelStreaming2StatNames::PliCount, FStat({ .Name = PixelStreaming2StatNames::PliCount }, 0.f));
+		Stats.Add(PixelStreaming2StatNames::NackCount, FStat({ .Name = PixelStreaming2StatNames::NackCount }, 0.f));
+		Stats.Add(PixelStreaming2StatNames::RetransmittedBytesReceived, FStat({ .Name = PixelStreaming2StatNames::RetransmittedBytesReceived }, 0.f));
+		Stats.Add(PixelStreaming2StatNames::RetransmittedPacketsReceived, FStat({ .Name = PixelStreaming2StatNames::RetransmittedPacketsReceived }, 0.f));
+		Stats.Add(PixelStreaming2StatNames::TotalEncodeBytesTarget, FStat({ .Name = PixelStreaming2StatNames::TotalEncodeBytesTarget }, 0.f));
+		Stats.Add(PixelStreaming2StatNames::KeyFramesDecoded, FStat({ .Name = PixelStreaming2StatNames::KeyFramesDecoded }, 0.f));
+		Stats.Add(PixelStreaming2StatNames::FrameWidth, FStat({ .Name = PixelStreaming2StatNames::FrameWidth }, 0.f));
+		Stats.Add(PixelStreaming2StatNames::FrameHeight, FStat({ .Name = PixelStreaming2StatNames::FrameHeight }, 0.f));
+		Stats.Add(PixelStreaming2StatNames::HugeFramesSent, FStat({ .Name = PixelStreaming2StatNames::HugeFramesSent }, 0.f));
+		Stats.Add(PixelStreaming2StatNames::PacketsLost, FStat({ .Name = PixelStreaming2StatNames::PacketsLost }, 0.f));
+		Stats.Add(PixelStreaming2StatNames::Jitter, FStat({ .Name = PixelStreaming2StatNames::Jitter }, 0.f));
+		Stats.Add(PixelStreaming2StatNames::RoundTripTime, FStat({ .Name = PixelStreaming2StatNames::RoundTripTime }, 0.f));
 
 		// These are values used to calculate extra values (stores time deltas etc)
-		AddNonRendered(PixelStreaming2StatNames::TargetBitrate);
-		AddNonRendered(PixelStreaming2StatNames::FramesSent);
-		AddNonRendered(PixelStreaming2StatNames::FramesReceived);
-		AddNonRendered(PixelStreaming2StatNames::BytesSent);
-		AddNonRendered(PixelStreaming2StatNames::BytesReceived);
-		AddNonRendered(PixelStreaming2StatNames::QPSum);
-		AddNonRendered(PixelStreaming2StatNames::TotalEncodeTime);
-		AddNonRendered(PixelStreaming2StatNames::FramesEncoded);
-		AddNonRendered(PixelStreaming2StatNames::FramesDecoded);
-		AddNonRendered(PixelStreaming2StatNames::TotalPacketSendDelay);
+		Stats.Add(PixelStreaming2StatNames::TargetBitrate, FStat({ .Name = PixelStreaming2StatNames::TargetBitrate, .DisplayFlags = EDisplayFlags::HIDDEN }, 0.f));
+		Stats.Add(PixelStreaming2StatNames::FramesSent, FStat({ .Name = PixelStreaming2StatNames::FramesSent, .DisplayFlags = EDisplayFlags::HIDDEN }, 0.f));
+		Stats.Add(PixelStreaming2StatNames::FramesReceived, FStat({ .Name = PixelStreaming2StatNames::FramesReceived, .DisplayFlags = EDisplayFlags::HIDDEN }, 0.f));
+		Stats.Add(PixelStreaming2StatNames::BytesSent, FStat({ .Name = PixelStreaming2StatNames::BytesSent, .DisplayFlags = EDisplayFlags::HIDDEN }, 0.f));
+		Stats.Add(PixelStreaming2StatNames::BytesReceived, FStat({ .Name = PixelStreaming2StatNames::BytesReceived, .DisplayFlags = EDisplayFlags::HIDDEN }, 0.f));
+		Stats.Add(PixelStreaming2StatNames::QPSum, FStat({ .Name = PixelStreaming2StatNames::QPSum, .DisplayFlags = EDisplayFlags::HIDDEN }, 0.f));
+		Stats.Add(PixelStreaming2StatNames::TotalEncodeTime, FStat({ .Name = PixelStreaming2StatNames::TotalEncodeTime, .DisplayFlags = EDisplayFlags::HIDDEN }, 0.f));
+		Stats.Add(PixelStreaming2StatNames::FramesEncoded, FStat({ .Name = PixelStreaming2StatNames::FramesEncoded, .DisplayFlags = EDisplayFlags::HIDDEN }, 0.f));
+		Stats.Add(PixelStreaming2StatNames::FramesDecoded, FStat({ .Name = PixelStreaming2StatNames::FramesDecoded, .DisplayFlags = EDisplayFlags::HIDDEN }, 0.f));
+		Stats.Add(PixelStreaming2StatNames::TotalPacketSendDelay, FStat({ .Name = PixelStreaming2StatNames::TotalPacketSendDelay, .DisplayFlags = EDisplayFlags::HIDDEN }, 0.f));
 
 		// Calculated stats below:
 		// FrameSent Per Second
-		AddStatCalculator([](FStatsSink& StatSource, double Period) -> TOptional<FStatData> {
-			FRTCTrackedStat* FramesSentStat = StatSource.Get(PixelStreaming2StatNames::FramesSent);
-			if (FramesSentStat && FramesSentStat->GetLatestStat().StatValue > 0)
+		Calculators.Add([](FStatsSink& StatSource, double Period) -> TOptional<FStat> {
+			FStat* FramesSentStat = StatSource.Get(PixelStreaming2StatNames::FramesSent);
+			if (FramesSentStat && FramesSentStat->GetValue<double>() > 0)
 			{
-				const double FramesSentPerSecond = FramesSentStat->CalculateDelta(Period);
-				FStatData	 FpsStat = FStatData(PixelStreaming2StatNames::FramesSentPerSecond, FramesSentPerSecond, 0);
-				FpsStat.DisplayFlags = FStatData::EDisplayFlags::TEXT | FStatData::EDisplayFlags::GRAPH;
-				return FpsStat;
+				const double FramesSentPerSecond = (FramesSentStat->GetValue<double>() - FramesSentStat->GetPrevValue<double>()) * Period;
+				return FStat({ .Name = PixelStreaming2StatNames::FramesSentPerSecond, .DisplayFlags = static_cast<EDisplayFlags>(EDisplayFlags::TEXT | EDisplayFlags::GRAPH) }, FramesSentPerSecond);
 			}
 			return {};
 		});
 
 		// FramesReceived Per Second
-		AddStatCalculator([](FStatsSink& StatSource, double Period) -> TOptional<FStatData> {
-			FRTCTrackedStat* FramesReceivedStat = StatSource.Get(PixelStreaming2StatNames::FramesReceived);
-			if (FramesReceivedStat && FramesReceivedStat->GetLatestStat().StatValue > 0)
+		Calculators.Add([](FStatsSink& StatSource, double Period) -> TOptional<FStat> {
+			FStat* FramesReceivedStat = StatSource.Get(PixelStreaming2StatNames::FramesReceived);
+			if (FramesReceivedStat && FramesReceivedStat->GetValue<double>() > 0)
 			{
-				const double FramesReceivedPerSecond = FramesReceivedStat->CalculateDelta(Period);
-				return FStatData(PixelStreaming2StatNames::FramesReceivedPerSecond, FramesReceivedPerSecond, 0);
+				const double FramesReceivedPerSecond = (FramesReceivedStat->GetValue<double>() - FramesReceivedStat->GetPrevValue<double>()) * Period;
+				return FStat({ .Name = PixelStreaming2StatNames::FramesReceivedPerSecond }, FramesReceivedPerSecond);
 			}
 			return {};
 		});
 
 		// Megabits sent Per Second
-		AddStatCalculator([](FStatsSink& StatSource, double Period) -> TOptional<FStatData> {
-			FRTCTrackedStat* BytesSentStat = StatSource.Get(PixelStreaming2StatNames::BytesSent);
-			if (BytesSentStat && BytesSentStat->GetLatestStat().StatValue > 0)
+		Calculators.Add([](FStatsSink& StatSource, double Period) -> TOptional<FStat> {
+			FStat* BytesSentStat = StatSource.Get(PixelStreaming2StatNames::BytesSent);
+			if (BytesSentStat && BytesSentStat->GetValue<double>() > 0)
 			{
-				const double BytesSentPerSecond = BytesSentStat->CalculateDelta(Period);
+				const double BytesSentPerSecond = (BytesSentStat->GetValue<double>() - BytesSentStat->GetPrevValue<double>()) * Period;
 				const double MegabitsPerSecond = BytesSentPerSecond / 1'000'000.0 * 8.0;
-				return FStatData(PixelStreaming2StatNames::BitrateMegabits, MegabitsPerSecond, 2);
+				return FStat({ .Name = PixelStreaming2StatNames::BitrateMegabits }, MegabitsPerSecond, 2);
 			}
 			return {};
 		});
 
 		// Bits sent Per Second
-		AddStatCalculator([](FStatsSink& StatSource, double Period) -> TOptional<FStatData> {
-			FRTCTrackedStat* BytesSentStat = StatSource.Get(PixelStreaming2StatNames::BytesSent);
-			if (BytesSentStat && BytesSentStat->GetLatestStat().StatValue > 0)
+		Calculators.Add([](FStatsSink& StatSource, double Period) -> TOptional<FStat> {
+			FStat* BytesSentStat = StatSource.Get(PixelStreaming2StatNames::BytesSent);
+			if (BytesSentStat && BytesSentStat->GetValue<double>() > 0)
 			{
-				const double BytesSentPerSecond = BytesSentStat->CalculateDelta(Period);
+				const double BytesSentPerSecond = (BytesSentStat->GetValue<double>() - BytesSentStat->GetPrevValue<double>()) * Period;
 				const double BitsPerSecond = BytesSentPerSecond * 8.0;
-				FStatData	 Stat = FStatData(PixelStreaming2StatNames::Bitrate, BitsPerSecond, 0);
-				Stat.DisplayFlags = FStatData::EDisplayFlags::HIDDEN; // We don't want to display bits per second (too many digits)
-				return Stat;
+				return FStat({ .Name = PixelStreaming2StatNames::Bitrate, .DisplayFlags = EDisplayFlags::HIDDEN }, BitsPerSecond);
 			}
 			return {};
 		});
 
 		// Target megabits sent Per Second
-		AddStatCalculator([](FStatsSink& StatSource, double Period) -> TOptional<FStatData> {
-			FRTCTrackedStat* TargetBpsStats = StatSource.Get(PixelStreaming2StatNames::TargetBitrate);
-			if (TargetBpsStats && TargetBpsStats->GetLatestStat().StatValue > 0)
+		Calculators.Add([](FStatsSink& StatSource, double Period) -> TOptional<FStat> {
+			FStat* TargetBpsStats = StatSource.Get(PixelStreaming2StatNames::TargetBitrate);
+			if (TargetBpsStats && TargetBpsStats->GetValue<double>() > 0)
 			{
-				const double TargetBps = TargetBpsStats->Average();
+				const double TargetBps = (TargetBpsStats->GetValue<double>() + TargetBpsStats->GetPrevValue<double>()) * 0.5f;
 				const double MegabitsPerSecond = TargetBps / 1'000'000.0;
-				return FStatData(PixelStreaming2StatNames::TargetBitrateMegabits, MegabitsPerSecond, 2);
+				return FStat({ .Name = PixelStreaming2StatNames::TargetBitrateMegabits }, MegabitsPerSecond, 2);
 			}
 			return {};
 		});
 
 		// Megabits received Per Second
-		AddStatCalculator([](FStatsSink& StatSource, double Period) -> TOptional<FStatData> {
-			FRTCTrackedStat* BytesReceivedStat = StatSource.Get(PixelStreaming2StatNames::BytesReceived);
-			if (BytesReceivedStat && BytesReceivedStat->GetLatestStat().StatValue > 0)
+		Calculators.Add([](FStatsSink& StatSource, double Period) -> TOptional<FStat> {
+			FStat* BytesReceivedStat = StatSource.Get(PixelStreaming2StatNames::BytesReceived);
+			if (BytesReceivedStat && BytesReceivedStat->GetValue<double>() > 0)
 			{
-				const double BytesReceivedPerSecond = BytesReceivedStat->CalculateDelta(Period);
+				const double BytesReceivedPerSecond = (BytesReceivedStat->GetValue<double>() - BytesReceivedStat->GetPrevValue<double>()) * Period;
 				const double MegabitsPerSecond = BytesReceivedPerSecond / 1000.0 * 8.0;
-				return FStatData(PixelStreaming2StatNames::Bitrate, MegabitsPerSecond, 2);
+				return FStat({ .Name = PixelStreaming2StatNames::Bitrate }, MegabitsPerSecond, 2);
 			}
 			return {};
 		});
 
 		// Encoded fps
-		AddStatCalculator([](FStatsSink& StatSource, double Period) -> TOptional<FStatData> {
-			FRTCTrackedStat* EncodedFramesStat = StatSource.Get(PixelStreaming2StatNames::FramesEncoded);
-			if (EncodedFramesStat && EncodedFramesStat->GetLatestStat().StatValue > 0)
+		Calculators.Add([](FStatsSink& StatSource, double Period) -> TOptional<FStat> {
+			FStat* EncodedFramesStat = StatSource.Get(PixelStreaming2StatNames::FramesEncoded);
+			if (EncodedFramesStat && EncodedFramesStat->GetValue<double>() > 0)
 			{
-				const double EncodedFramesPerSecond = EncodedFramesStat->CalculateDelta(Period);
-				return FStatData(PixelStreaming2StatNames::EncodedFramesPerSecond, EncodedFramesPerSecond, 0);
+				const double EncodedFramesPerSecond = (EncodedFramesStat->GetValue<double>() - EncodedFramesStat->GetPrevValue<double>()) * Period;
+				return FStat({ .Name = PixelStreaming2StatNames::EncodedFramesPerSecond }, EncodedFramesPerSecond);
 			}
 			return {};
 		});
 
 		// Decoded fps
-		AddStatCalculator([](FStatsSink& StatSource, double Period) -> TOptional<FStatData> {
-			FRTCTrackedStat* DecodedFramesStat = StatSource.Get(PixelStreaming2StatNames::FramesDecoded);
-			if (DecodedFramesStat && DecodedFramesStat->GetLatestStat().StatValue > 0)
+		Calculators.Add([](FStatsSink& StatSource, double Period) -> TOptional<FStat> {
+			FStat* DecodedFramesStat = StatSource.Get(PixelStreaming2StatNames::FramesDecoded);
+			if (DecodedFramesStat && DecodedFramesStat->GetValue<double>() > 0)
 			{
-				const double DecodedFramesPerSecond = DecodedFramesStat->CalculateDelta(Period);
-				return FStatData(PixelStreaming2StatNames::DecodedFramesPerSecond, DecodedFramesPerSecond, 0);
+				const double DecodedFramesPerSecond = (DecodedFramesStat->GetValue<double>() - DecodedFramesStat->GetPrevValue<double>()) * Period;
+				return FStat({ .Name = PixelStreaming2StatNames::DecodedFramesPerSecond }, DecodedFramesPerSecond);
 			}
 			return {};
 		});
 
 		// Avg QP Per Second
-		AddStatCalculator([](FStatsSink& StatSource, double Period) -> TOptional<FStatData> {
-			FRTCTrackedStat* QPSumStat = StatSource.Get(PixelStreaming2StatNames::QPSum);
-			FStatData*		 EncodedFramesPerSecond = StatSource.GetCalculatedStat(PixelStreaming2StatNames::EncodedFramesPerSecond);
-			if (QPSumStat && QPSumStat->GetLatestStat().StatValue > 0
-				&& EncodedFramesPerSecond && EncodedFramesPerSecond->StatValue > 0.0)
+		Calculators.Add([](FStatsSink& StatSource, double Period) -> TOptional<FStat> {
+			FStat* QPSumStat = StatSource.Get(PixelStreaming2StatNames::QPSum);
+			FStat* EncodedFramesPerSecond = StatSource.Get(PixelStreaming2StatNames::EncodedFramesPerSecond);
+			if (QPSumStat && QPSumStat->GetValue<double>() > 0
+				&& EncodedFramesPerSecond && EncodedFramesPerSecond->GetValue<double>() > 0.0)
 			{
-				const double QPSumDeltaPerSecond = QPSumStat->CalculateDelta(Period);
-				const double MeanQPPerFrame = QPSumDeltaPerSecond / EncodedFramesPerSecond->StatValue;
-				FName		 StatName = PixelStreaming2StatNames::MeanQPPerSecond;
-				return FStatData(StatName, MeanQPPerFrame, 0);
+				const double QPSumDeltaPerSecond = (QPSumStat->GetValue<double>() - QPSumStat->GetPrevValue<double>()) * Period;
+				const double MeanQPPerFrame = QPSumDeltaPerSecond / EncodedFramesPerSecond->GetValue<double>();
+				return FStat({ .Name = PixelStreaming2StatNames::MeanQPPerSecond }, MeanQPPerFrame);
 			}
 			return {};
 		});
 
 		// Mean EncodeTime (ms) Per Frame
-		AddStatCalculator([](FStatsSink& StatSource, double Period) -> TOptional<FStatData> {
-			FRTCTrackedStat* TotalEncodeTimeStat = StatSource.Get(PixelStreaming2StatNames::TotalEncodeTime);
-			FStatData*		 EncodedFramesPerSecond = StatSource.GetCalculatedStat(PixelStreaming2StatNames::EncodedFramesPerSecond);
-			if (TotalEncodeTimeStat && TotalEncodeTimeStat->GetLatestStat().StatValue > 0
-				&& EncodedFramesPerSecond && EncodedFramesPerSecond->StatValue > 0.0)
+		Calculators.Add([](FStatsSink& StatSource, double Period) -> TOptional<FStat> {
+			FStat* TotalEncodeTimeStat = StatSource.Get(PixelStreaming2StatNames::TotalEncodeTime);
+			FStat* EncodedFramesPerSecond = StatSource.Get(PixelStreaming2StatNames::EncodedFramesPerSecond);
+			if (TotalEncodeTimeStat && TotalEncodeTimeStat->GetValue<double>() > 0
+				&& EncodedFramesPerSecond && EncodedFramesPerSecond->GetValue<double>() > 0.0)
 			{
-				const double TotalEncodeTimePerSecond = TotalEncodeTimeStat->CalculateDelta(Period);
-				const double MeanEncodeTimePerFrameMs = TotalEncodeTimePerSecond / EncodedFramesPerSecond->StatValue * 1000.0;
-				return FStatData(PixelStreaming2StatNames::MeanEncodeTime, MeanEncodeTimePerFrameMs, 2);
+				const double TotalEncodeTimePerSecond = (TotalEncodeTimeStat->GetValue<double>() - TotalEncodeTimeStat->GetPrevValue<double>()) * Period;
+				const double MeanEncodeTimePerFrameMs = TotalEncodeTimePerSecond / EncodedFramesPerSecond->GetValue<double>() * 1000.0;
+				return FStat({ .Name = PixelStreaming2StatNames::MeanEncodeTime }, MeanEncodeTimePerFrameMs, 2);
 			}
 			return {};
 		});
 
 		// Mean SendDelay (ms) Per Frame
-		AddStatCalculator([](FStatsSink& StatSource, double Period) -> TOptional<FStatData> {
-			FRTCTrackedStat* TotalSendDelayStat = StatSource.Get(PixelStreaming2StatNames::TotalPacketSendDelay);
-			FStatData*		 FramesSentPerSecond = StatSource.GetCalculatedStat(PixelStreaming2StatNames::FramesSentPerSecond);
-			if (TotalSendDelayStat && TotalSendDelayStat->GetLatestStat().StatValue > 0
-				&& FramesSentPerSecond && FramesSentPerSecond->StatValue > 0.0)
+		Calculators.Add([](FStatsSink& StatSource, double Period) -> TOptional<FStat> {
+			FStat* TotalSendDelayStat = StatSource.Get(PixelStreaming2StatNames::TotalPacketSendDelay);
+			FStat* TotalPacketsSent = StatSource.Get(PixelStreaming2StatNames::PacketsSent);
+			if (TotalSendDelayStat && TotalSendDelayStat->GetValue<double>() > 0.0
+				&& TotalPacketsSent && TotalPacketsSent->GetValue<double>() > 0.0)
 			{
-				const double TotalSendDelayPerSecond = TotalSendDelayStat->CalculateDelta(Period);
-				const double MeanSendDelayPerFrameMs = TotalSendDelayPerSecond / FramesSentPerSecond->StatValue * 1000.0;
-				return FStatData(PixelStreaming2StatNames::MeanSendDelay, MeanSendDelayPerFrameMs, 2);
+				const double MeanSendDelayPerFrameMs = (TotalSendDelayStat->GetValue<double>() / TotalPacketsSent->GetValue<double>()) * 1000.0;
+				return FStat({ .Name = PixelStreaming2StatNames::MeanSendDelay }, MeanSendDelayPerFrameMs, 2);
 			}
 			return {};
 		});
 
 		// JitterBufferDelay (ms)
-		AddStatCalculator([](FStatsSink& StatSource, double Period) -> TOptional<FStatData> {
-			FRTCTrackedStat* JitterBufferDelayStat = StatSource.Get(PixelStreaming2StatNames::JitterBufferDelay);
-			FStatData*		 FramesReceivedPerSecond = StatSource.GetCalculatedStat(PixelStreaming2StatNames::FramesReceivedPerSecond);
-			if (JitterBufferDelayStat && JitterBufferDelayStat->GetLatestStat().StatValue > 0
-				&& FramesReceivedPerSecond && FramesReceivedPerSecond->StatValue > 0.0)
+		Calculators.Add([](FStatsSink& StatSource, double Period) -> TOptional<FStat> {
+			FStat* JitterBufferDelayStat = StatSource.Get(PixelStreaming2StatNames::JitterBufferDelay);
+			FStat* FramesReceivedPerSecond = StatSource.Get(PixelStreaming2StatNames::FramesReceivedPerSecond);
+			if (JitterBufferDelayStat && JitterBufferDelayStat->GetValue<double>() > 0
+				&& FramesReceivedPerSecond && FramesReceivedPerSecond->GetValue<double>() > 0.0)
 			{
-				const double TotalJitterBufferDelayPerSecond = JitterBufferDelayStat->CalculateDelta(Period);
-				const double MeanJitterBufferDelayMs = TotalJitterBufferDelayPerSecond / FramesReceivedPerSecond->StatValue * 1000.0;
-				return FStatData(PixelStreaming2StatNames::JitterBufferDelay, MeanJitterBufferDelayMs, 2);
+				const double TotalJitterBufferDelayPerSecond = (JitterBufferDelayStat->GetValue<double>() - JitterBufferDelayStat->GetPrevValue<double>()) * Period;
+				const double MeanJitterBufferDelayMs = TotalJitterBufferDelayPerSecond / FramesReceivedPerSecond->GetValue<double>() * 1000.0;
+				return FStat({ .Name = PixelStreaming2StatNames::JitterBufferDelay }, MeanJitterBufferDelayMs, 2);
 			}
 			return {};
 		});
@@ -922,73 +1254,78 @@ namespace UE::PixelStreaming2
 			return;
 		}
 
-		for (TPair<FName, FRTCTrackedStat>& Tuple : Stats)
+		for (TPair<FName, FStat>& Tuple : Stats)
 		{
-			double NewValue = 0;
+			FStatVariant NewValue;
 			if (Tuple.Key == PixelStreaming2StatNames::FirCount)
 			{
-				NewValue = InStats._local._firCount;
+				NewValue = FStatVariant(TInPlaceType<double>(), InStats._local._firCount);
 			}
 			else if (Tuple.Key == PixelStreaming2StatNames::PliCount)
 			{
-				NewValue = InStats._local._pliCount;
+				NewValue = FStatVariant(TInPlaceType<double>(), InStats._local._pliCount);
 			}
 			else if (Tuple.Key == PixelStreaming2StatNames::NackCount)
 			{
-				NewValue = InStats._local._nackCount;
+				NewValue = FStatVariant(TInPlaceType<double>(), InStats._local._nackCount);
 			}
 			else if (Tuple.Key == PixelStreaming2StatNames::RetransmittedBytesReceived)
 			{
-				NewValue = InStats._local._retransmittedBytesReceived;
+				NewValue = FStatVariant(TInPlaceType<double>(), InStats._local._retransmittedBytesReceived);
 			}
 			else if (Tuple.Key == PixelStreaming2StatNames::RetransmittedPacketsReceived)
 			{
-				NewValue = InStats._local._retransmittedPacketsReceived;
+				NewValue = FStatVariant(TInPlaceType<double>(), InStats._local._retransmittedPacketsReceived);
 			}
 			else if (Tuple.Key == PixelStreaming2StatNames::KeyFramesDecoded)
 			{
-				NewValue = InStats._local._keyFramesDecoded;
+				NewValue = FStatVariant(TInPlaceType<double>(), InStats._local._keyFramesDecoded);
 			}
 			else if (Tuple.Key == PixelStreaming2StatNames::FrameWidth)
 			{
-				NewValue = InStats._local._frameWidth;
+				NewValue = FStatVariant(TInPlaceType<double>(), InStats._local._frameWidth);
 			}
 			else if (Tuple.Key == PixelStreaming2StatNames::FrameHeight)
 			{
-				NewValue = InStats._local._frameHeight;
+				NewValue = FStatVariant(TInPlaceType<double>(), InStats._local._frameHeight);
 			}
 			else if (Tuple.Key == PixelStreaming2StatNames::FramesReceived)
 			{
-				NewValue = InStats._local._framesReceived;
+				NewValue = FStatVariant(TInPlaceType<double>(), InStats._local._framesReceived);
 			}
 			else if (Tuple.Key == PixelStreaming2StatNames::BytesReceived)
 			{
-				NewValue = InStats._local._bytesReceived;
+				NewValue = FStatVariant(TInPlaceType<double>(), InStats._local._bytesReceived);
 			}
 			else if (Tuple.Key == PixelStreaming2StatNames::QPSum)
 			{
-				NewValue = InStats._local._qpSum;
+				NewValue = FStatVariant(TInPlaceType<double>(), InStats._local._qpSum);
 			}
 			else if (Tuple.Key == PixelStreaming2StatNames::FramesDecoded)
 			{
-				NewValue = InStats._local._framesDecoded;
+				NewValue = FStatVariant(TInPlaceType<double>(), InStats._local._framesDecoded);
 			}
 			else if (Tuple.Key == PixelStreaming2StatNames::PacketsLost)
 			{
-				NewValue = InStats._local._packetsLost;
+				NewValue = FStatVariant(TInPlaceType<double>(), InStats._local._packetsLost);
 			}
 			else if (Tuple.Key == PixelStreaming2StatNames::Jitter)
 			{
-				NewValue = InStats._local._jitter;
+				NewValue = FStatVariant(TInPlaceType<double>(), InStats._local._jitter);
 			}
 			else if (Tuple.Key == PixelStreaming2StatNames::RoundTripTime)
 			{
-				NewValue = InStats._remote._roundTripTime;
+				NewValue = FStatVariant(TInPlaceType<double>(), InStats._remote._roundTripTime);
 			}
 
-			if (UpdateValue(NewValue, &Tuple.Value))
+			if (NewValue.GetIndex() == FStatVariant::IndexOfType<TYPE_OF_NULLPTR>())
 			{
-				PSStats->StorePeerStat(PeerId, Category, Tuple.Value.GetLatestStat());
+				continue;
+			}
+
+			if (Tuple.Value.SetValue(NewValue))
+			{
+				PSStats->StorePeerStat(PeerId, Category, Tuple.Value);
 			}
 		}
 
@@ -1002,7 +1339,7 @@ namespace UE::PixelStreaming2
 		: FStatsSink(InCategory)
 	{
 		// Track video source fps
-		Add(PixelStreaming2StatNames::SourceFps, 0);
+		Stats.Add(PixelStreaming2StatNames::SourceFps, FStat({ .Name = PixelStreaming2StatNames::SourceFps }, 0.f));
 	}
 
 	void FRTCStatsCollector::FVideoSourceStatsSink::Process(const EpicRtcVideoSourceStats& InStats, const FString& PeerId, double SecondsDelta)
@@ -1013,17 +1350,62 @@ namespace UE::PixelStreaming2
 			return;
 		}
 
-		for (TPair<FName, FRTCTrackedStat>& Tuple : Stats)
+		for (TPair<FName, FStat>& Tuple : Stats)
 		{
-			double NewValue = 0;
+			FStatVariant NewValue;
 			if (Tuple.Key == PixelStreaming2StatNames::SourceFps)
 			{
-				NewValue = InStats._framesPerSecond;
+				NewValue = FStatVariant(TInPlaceType<double>(), InStats._framesPerSecond);
+			}
+
+			if (NewValue.GetIndex() == FStatVariant::IndexOfType<TYPE_OF_NULLPTR>())
+			{
+				continue;
+			}
+
+			if (Tuple.Value.SetValue(NewValue))
+			{
+				PSStats->StorePeerStat(PeerId, Category, Tuple.Value);
+			}
+		}
+
+		PostProcess(PSStats, PeerId, SecondsDelta);
+	}
+
+	/**
+	 * ---------- FRTPVideoCodecSink ----------
+	 */
+	FRTCStatsCollector::FVideoCodecStatsSink::FVideoCodecStatsSink(FName InCategory)
+		: FStatsSink(InCategory)
+	{
+		// Track video source fps
+		Stats.Add(PixelStreaming2StatNames::MimeType, FStat({ .Name = PixelStreaming2StatNames::MimeType }, FString(TEXT(""))));
+	}
+
+	void FRTCStatsCollector::FVideoCodecStatsSink::Process(const EpicRtcCodecStats& InStats, const FString& PeerId, double SecondsDelta)
+	{
+		FStats* PSStats = FStats::Get();
+		if (!PSStats)
+		{
+			return;
+		}
+
+		for (TPair<FName, FStat>& Tuple : Stats)
+		{
+			FStatVariant NewValue;
+			if (Tuple.Key == PixelStreaming2StatNames::MimeType)
+			{
+				NewValue = FStatVariant(TInPlaceType<FString>(), ToString(InStats._mimeType));
+			}
+
+			if (NewValue.GetIndex() == FStatVariant::IndexOfType<TYPE_OF_NULLPTR>())
+			{
+				continue;
 			}
 
-			if (UpdateValue(NewValue, &Tuple.Value))
+			if (Tuple.Value.SetValue(NewValue))
 			{
-				PSStats->StorePeerStat(PeerId, Category, Tuple.Value.GetLatestStat());
+				PSStats->StorePeerStat(PeerId, Category, Tuple.Value);
 			}
 		}
 
@@ -1036,8 +1418,8 @@ namespace UE::PixelStreaming2
 	FRTCStatsCollector::FAudioSourceStatsSink::FAudioSourceStatsSink(FName InCategory)
 		: FStatsSink(InCategory)
 	{
-		Add(PixelStreaming2StatNames::AudioLevel, 0);
-		Add(PixelStreaming2StatNames::TotalSamplesDuration, 0);
+		Stats.Add(PixelStreaming2StatNames::AudioLevel, FStat({ .Name = PixelStreaming2StatNames::AudioLevel }, 0.f));
+		Stats.Add(PixelStreaming2StatNames::TotalSamplesDuration, FStat({ .Name = PixelStreaming2StatNames::TotalSamplesDuration }, 0.f));
 	}
 
 	void FRTCStatsCollector::FAudioSourceStatsSink::Process(const EpicRtcAudioSourceStats& InStats, const FString& PeerId, double SecondsDelta)
@@ -1048,21 +1430,76 @@ namespace UE::PixelStreaming2
 			return;
 		}
 
-		for (TPair<FName, FRTCTrackedStat>& Tuple : Stats)
+		for (TPair<FName, FStat>& Tuple : Stats)
 		{
-			double NewValue = 0;
+			FStatVariant NewValue;
 			if (Tuple.Key == PixelStreaming2StatNames::AudioLevel)
 			{
-				NewValue = InStats._audioLevel;
+				NewValue = FStatVariant(TInPlaceType<double>(), InStats._audioLevel);
 			}
 			else if (Tuple.Key == PixelStreaming2StatNames::TotalSamplesDuration)
 			{
-				NewValue = InStats._totalSamplesDuration;
+				NewValue = FStatVariant(TInPlaceType<double>(), InStats._totalSamplesDuration);
+			}
+
+			if (NewValue.GetIndex() == FStatVariant::IndexOfType<TYPE_OF_NULLPTR>())
+			{
+				continue;
 			}
 
-			if (UpdateValue(NewValue, &Tuple.Value))
+			if (Tuple.Value.SetValue(NewValue))
 			{
-				PSStats->StorePeerStat(PeerId, Category, Tuple.Value.GetLatestStat());
+				PSStats->StorePeerStat(PeerId, Category, Tuple.Value);
+			}
+		}
+
+		PostProcess(PSStats, PeerId, SecondsDelta);
+	}
+
+	/**
+	 * ---------- FRTPAudioCodecSink ----------
+	 */
+	FRTCStatsCollector::FAudioCodecStatsSink::FAudioCodecStatsSink(FName InCategory)
+		: FStatsSink(InCategory)
+	{
+		// Track video source fps
+		Stats.Add(PixelStreaming2StatNames::MimeType, FStat({ .Name = PixelStreaming2StatNames::MimeType }, FString(TEXT(""))));
+		Stats.Add(PixelStreaming2StatNames::Channels, FStat({ .Name = PixelStreaming2StatNames::Channels }, 0.f));
+		Stats.Add(PixelStreaming2StatNames::ClockRate, FStat({ .Name = PixelStreaming2StatNames::ClockRate }, 0.f));
+	}
+
+	void FRTCStatsCollector::FAudioCodecStatsSink::Process(const EpicRtcCodecStats& InStats, const FString& PeerId, double SecondsDelta)
+	{
+		FStats* PSStats = FStats::Get();
+		if (!PSStats)
+		{
+			return;
+		}
+
+		for (TPair<FName, FStat>& Tuple : Stats)
+		{
+			FStatVariant NewValue;
+			if (Tuple.Key == PixelStreaming2StatNames::MimeType)
+			{
+				NewValue = FStatVariant(TInPlaceType<FString>(), ToString(InStats._mimeType));
+			}
+			else if (Tuple.Key == PixelStreaming2StatNames::Channels)
+			{
+				NewValue = FStatVariant(TInPlaceType<double>(), InStats._channels);
+			}
+			else if (Tuple.Key == PixelStreaming2StatNames::ClockRate)
+			{
+				NewValue = FStatVariant(TInPlaceType<double>(), InStats._clockRate);
+			}
+
+			if (NewValue.GetIndex() == FStatVariant::IndexOfType<TYPE_OF_NULLPTR>())
+			{
+				continue;
+			}
+
+			if (Tuple.Value.SetValue(NewValue))
+			{
+				PSStats->StorePeerStat(PeerId, Category, Tuple.Value);
 			}
 		}
 
@@ -1077,10 +1514,10 @@ namespace UE::PixelStreaming2
 	{
 		// These names are added as aliased names because `bytesSent` is ambiguous stat that is used across inbound-rtp, outbound-rtp, and data-channel
 		// so to disambiguate which state we are referring to we record the `bytesSent` stat for the data-channel but store and report it as `data-channel-bytesSent`
-		AddAliased(PixelStreaming2StatNames::MessagesSent, PixelStreaming2StatNames::DataChannelMessagesSent, 0, FStatData::EDisplayFlags::TEXT | FStatData::EDisplayFlags::GRAPH);
-		AddAliased(PixelStreaming2StatNames::MessagesReceived, PixelStreaming2StatNames::DataChannelBytesReceived, 0, FStatData::EDisplayFlags::TEXT | FStatData::EDisplayFlags::GRAPH);
-		AddAliased(PixelStreaming2StatNames::BytesSent, PixelStreaming2StatNames::DataChannelBytesSent, 0, FStatData::EDisplayFlags::TEXT | FStatData::EDisplayFlags::GRAPH);
-		AddAliased(PixelStreaming2StatNames::BytesReceived, PixelStreaming2StatNames::DataChannelMessagesReceived, 0, FStatData::EDisplayFlags::TEXT | FStatData::EDisplayFlags::GRAPH);
+		Stats.Add(PixelStreaming2StatNames::MessagesSent, FStat({ .Name = PixelStreaming2StatNames::MessagesSent, .Alias = PixelStreaming2StatNames::DataChannelMessagesSent, .DisplayFlags = static_cast<EDisplayFlags>(EDisplayFlags::TEXT | EDisplayFlags::GRAPH) }, 0.f));
+		Stats.Add(PixelStreaming2StatNames::MessagesReceived, FStat({ .Name = PixelStreaming2StatNames::MessagesReceived, .Alias = PixelStreaming2StatNames::DataChannelBytesReceived, .DisplayFlags = static_cast<EDisplayFlags>(EDisplayFlags::TEXT | EDisplayFlags::GRAPH) }, 0.f));
+		Stats.Add(PixelStreaming2StatNames::BytesSent, FStat({ .Name = PixelStreaming2StatNames::BytesSent, .Alias = PixelStreaming2StatNames::DataChannelBytesSent, .DisplayFlags = static_cast<EDisplayFlags>(EDisplayFlags::TEXT | EDisplayFlags::GRAPH) }, 0.f));
+		Stats.Add(PixelStreaming2StatNames::BytesReceived, FStat({ .Name = PixelStreaming2StatNames::BytesReceived, .Alias = PixelStreaming2StatNames::DataChannelMessagesReceived, .DisplayFlags = static_cast<EDisplayFlags>(EDisplayFlags::TEXT | EDisplayFlags::GRAPH) }, 0.f));
 	}
 
 	void FRTCStatsCollector::FDataTrackStatsSink::Process(const EpicRtcDataTrackStats& InStats, const FString& PeerId, double SecondsDelta)
@@ -1091,29 +1528,34 @@ namespace UE::PixelStreaming2
 			return;
 		}
 
-		for (TPair<FName, FRTCTrackedStat>& Tuple : Stats)
+		for (TPair<FName, FStat>& Tuple : Stats)
 		{
-			double NewValue = 0;
+			FStatVariant NewValue;
 			if (Tuple.Key == PixelStreaming2StatNames::MessagesSent)
 			{
-				NewValue = InStats._messagesSent;
+				NewValue = FStatVariant(TInPlaceType<double>(), InStats._messagesSent);
 			}
 			else if (Tuple.Key == PixelStreaming2StatNames::MessagesReceived)
 			{
-				NewValue = InStats._messagesReceived;
+				NewValue = FStatVariant(TInPlaceType<double>(), InStats._messagesReceived);
 			}
 			else if (Tuple.Key == PixelStreaming2StatNames::BytesSent)
 			{
-				NewValue = InStats._bytesSent;
+				NewValue = FStatVariant(TInPlaceType<double>(), InStats._bytesSent);
 			}
 			else if (Tuple.Key == PixelStreaming2StatNames::BytesReceived)
 			{
-				NewValue = InStats._bytesReceived;
+				NewValue = FStatVariant(TInPlaceType<double>(), InStats._bytesReceived);
+			}
+
+			if (NewValue.GetIndex() == FStatVariant::IndexOfType<TYPE_OF_NULLPTR>())
+			{
+				continue;
 			}
 
-			if (UpdateValue(NewValue, &Tuple.Value))
+			if (Tuple.Value.SetValue(NewValue))
 			{
-				PSStats->StorePeerStat(PeerId, Category, Tuple.Value.GetLatestStat());
+				PSStats->StorePeerStat(PeerId, Category, Tuple.Value);
 			}
 		}
 
@@ -1126,8 +1568,8 @@ namespace UE::PixelStreaming2
 	FRTCStatsCollector::FCandidatePairStatsSink::FCandidatePairStatsSink(FName InCategory)
 		: FStatsSink(InCategory)
 	{
-		Add(PixelStreaming2StatNames::AvailableOutgoingBitrate, 0);
-		Add(PixelStreaming2StatNames::AvailableIncomingBitrate, 0);
+		Stats.Add(PixelStreaming2StatNames::AvailableOutgoingBitrate, FStat({ .Name = PixelStreaming2StatNames::AvailableOutgoingBitrate }, 0.f));
+		Stats.Add(PixelStreaming2StatNames::AvailableIncomingBitrate, FStat({ .Name = PixelStreaming2StatNames::AvailableIncomingBitrate }, 0.f));
 	}
 
 	void FRTCStatsCollector::FCandidatePairStatsSink::Process(const EpicRtcIceCandidatePairStats& InStats, const FString& PeerId, double SecondsDelta)
@@ -1138,53 +1580,32 @@ namespace UE::PixelStreaming2
 			return;
 		}
 
-		for (TPair<FName, FRTCTrackedStat>& Tuple : Stats)
+		for (TPair<FName, FStat>& Tuple : Stats)
 		{
-			double NewValue = 0;
+			FStatVariant NewValue;
 			if (Tuple.Key == PixelStreaming2StatNames::AvailableOutgoingBitrate)
 			{
-				NewValue = InStats._availableOutgoingBitrate;
+				NewValue = FStatVariant(TInPlaceType<double>(), InStats._availableOutgoingBitrate);
 			}
 			else if (Tuple.Key == PixelStreaming2StatNames::AvailableIncomingBitrate)
 			{
-				NewValue = InStats._availableIncomingBitrate;
+				NewValue = FStatVariant(TInPlaceType<double>(), InStats._availableIncomingBitrate);
 			}
 
-			if (UpdateValue(NewValue, &Tuple.Value))
+			if (NewValue.GetIndex() == FStatVariant::IndexOfType<TYPE_OF_NULLPTR>())
 			{
-				PSStats->StorePeerStat(PeerId, Category, Tuple.Value.GetLatestStat());
+				continue;
+			}
+
+			if (Tuple.Value.SetValue(NewValue))
+			{
+				PSStats->StorePeerStat(PeerId, Category, Tuple.Value);
 			}
 		}
 
 		PostProcess(PSStats, PeerId, SecondsDelta);
 	}
 
-	/**
-	 * ---------- FRTCTrackedStat -------------------
-	 */
-	FRTCStatsCollector::FRTCTrackedStat::FRTCTrackedStat(FName StatName, FName Alias, int NDecimalPlaces, uint8 DisplayFlags)
-		: LatestStat(StatName, 0.0, NDecimalPlaces)
-	{
-		LatestStat.DisplayFlags = DisplayFlags;
-		LatestStat.Alias = Alias;
-	}
-
-	double FRTCStatsCollector::FRTCTrackedStat::CalculateDelta(double Period) const
-	{
-		return (LatestStat.StatValue - PrevValue) * Period;
-	}
-
-	double FRTCStatsCollector::FRTCTrackedStat::Average() const
-	{
-		return (LatestStat.StatValue + PrevValue) * 0.5;
-	}
-
-	void FRTCStatsCollector::FRTCTrackedStat::SetLatestValue(double InValue)
-	{
-		PrevValue = LatestStat.StatValue;
-		LatestStat.StatValue = InValue;
-	}
-
 	/**
 	 * --------- FStatsSink ------------------------
 	 */
@@ -1193,29 +1614,15 @@ namespace UE::PixelStreaming2
 	{
 	}
 
-	void FRTCStatsCollector::FStatsSink::AddAliased(FName StatName, FName AliasedName, int NDecimalPlaces, uint8 DisplayFlags)
-	{
-		FRTCTrackedStat Stat = FRTCTrackedStat(StatName, AliasedName, NDecimalPlaces, DisplayFlags);
-		Stats.Add(StatName, Stat);
-	}
-
-	bool FRTCStatsCollector::FStatsSink::UpdateValue(double NewValue, FRTCTrackedStat* SetValueHere)
-	{
-		const bool bZeroInitially = SetValueHere->GetLatestStat().StatValue == 0.0;
-		SetValueHere->SetLatestValue(NewValue);
-		const bool bZeroStill = SetValueHere->GetLatestStat().StatValue == 0.0;
-		return !(bZeroInitially && bZeroStill);
-	}
-
 	void FRTCStatsCollector::FStatsSink::PostProcess(FStats* PSStats, const FString& PeerId, double SecondsDelta)
 	{
 		for (auto& Calculator : Calculators)
 		{
-			TOptional<FStatData> OptStatData = Calculator(*this, SecondsDelta);
+			TOptional<FStat> OptStatData = Calculator(*this, SecondsDelta);
 			if (OptStatData.IsSet())
 			{
-				FStatData& StatData = *OptStatData;
-				CalculatedStats.Add(StatData.StatName, StatData);
+				FStat& StatData = *OptStatData;
+				Stats.Add(StatData.GetName(), StatData);
 				PSStats->StorePeerStat(PeerId, Category, StatData);
 			}
 		}
diff --git a/Engine/Plugins/Media/PixelStreaming2/Source/PixelStreaming2RTC/Private/Stats.cpp b/Engine/Plugins/Media/PixelStreaming2/Source/PixelStreaming2RTC/Private/Stats.cpp
--- a/Engine/Plugins/Media/PixelStreaming2/Source/PixelStreaming2RTC/Private/Stats.cpp
+++ b/Engine/Plugins/Media/PixelStreaming2/Source/PixelStreaming2RTC/Private/Stats.cpp
@@ -48,9 +48,9 @@ namespace UE::PixelStreaming2
 		FCoreDelegates::OnPostEngineInit.AddRaw(this, &FStats::RegisterEngineHooks);
 	}
 
-	void FStats::StorePeerStat(const FString& PlayerId, FName StatCategory, FStatData Stat)
+	void FStats::StorePeerStat(const FString& PlayerId, FName StatCategory, FStat Stat)
 	{
-		FName& StatName = Stat.Alias.IsSet() ? Stat.Alias.GetValue() : Stat.StatName;
+		FName StatName = Stat.GetDisplayName();
 
 		bool Updated = false;
 		{
@@ -71,21 +71,24 @@ namespace UE::PixelStreaming2
 		{
 			if (Stat.ShouldGraph())
 			{
-				GraphValue(StatName, Stat.StatValue, 60, 0, Stat.StatValue * 10.0f, 0);
+				GraphValue(StatName, Stat.GetValue<double>(), 60, 0, Stat.GetValue<double>() * 10.0f, 0);
 			}
 
-			// If a stat has an alias, use that as the storage key, otherwise use its display name
-			FireStatChanged(PlayerId, StatName, Stat.StatValue);
+			if(Stat.IsNumeric())
+			{
+				// If a stat has an alias, use that as the storage key, otherwise use its underlying name
+				FireStatChanged(PlayerId, StatName, Stat.GetValue<double>());
+			}
 		}
 	}
 
-	bool FStats::QueryPeerStat(const FString& PlayerId, FName InStatCategory, FName StatToQuery, double& OutValue) const
+	bool FStats::QueryPeerStat(const FString& PlayerId, FName InStatCategory, FName StatToQuery, double& OutValue)
 	{
 		FScopeLock Lock(&PeerStatsCS);
 
-		if (const FPeerStats* SinglePeerStats = PeerStats.Find(PlayerId))
+		if (FPeerStats* SinglePeerStats = PeerStats.Find(PlayerId))
 		{
-			const TMap<FName, FStatGroup>& StatGroups = SinglePeerStats->GetStatGroups();
+			TMap<FName, FStatGroup>& StatGroups = SinglePeerStats->GetStatGroups();
 			TArray<FName>				   StatCategories;
 			StatGroups.GetKeys(StatCategories);
 
@@ -139,29 +142,14 @@ namespace UE::PixelStreaming2
 		}
 	}
 
-	double CalcMA(double PrevAvg, int NumSamples, double Value)
-	{
-		const double Result = NumSamples * PrevAvg + Value;
-		return Result / (PrevAvg + 1.0);
-	}
-
-	double CalcEMA(double PrevAvg, int NumSamples, double Value)
+	void FStats::StoreApplicationStat(FStat Stat)
 	{
-		const double Mult = 2.0 / (NumSamples + 1.0);
-		const double Result = (Value - PrevAvg) * Mult + PrevAvg;
-		return Result;
-	}
-
-	void FStats::StoreApplicationStat(FStatData Stat)
-	{
-		bool bUpdated = false;
-
-		// If a stat has an alias, use that as the storage key, otherwise use its display name
-		FName& StatName = Stat.Alias.IsSet() ? Stat.Alias.GetValue() : Stat.StatName;
+		// If a stat has an alias, use that as the storage key, otherwise use its underlying name
+		FName StatName = Stat.GetDisplayName();
 
 		if (Stat.ShouldGraph())
 		{
-			GraphValue(StatName, Stat.StatValue, 60, 0, Stat.StatValue, 0);
+			GraphValue(StatName, Stat.GetValue<double>(), 60, 0, Stat.GetValue<double>(), 0);
 		}
 
 		{
@@ -169,54 +157,32 @@ namespace UE::PixelStreaming2
 
 			if (ApplicationStats.Contains(StatName))
 			{
-				FStoredStat* StoredStat = ApplicationStats.Find(StatName);
-
-				if (Stat.bSmooth && StoredStat->Stat.StatValue != 0)
-				{
-					const int MaxSamples = 60;
-					StoredStat->Stat.NumSamples = FGenericPlatformMath::Min(MaxSamples, StoredStat->Stat.NumSamples + 1);
-					if (StoredStat->Stat.NumSamples < MaxSamples)
-					{
-						StoredStat->Stat.LastEMA = CalcMA(StoredStat->Stat.LastEMA, StoredStat->Stat.NumSamples - 1, Stat.StatValue);
-					}
-					else
-					{
-						StoredStat->Stat.LastEMA = CalcEMA(StoredStat->Stat.LastEMA, StoredStat->Stat.NumSamples - 1, Stat.StatValue);
-					}
-					StoredStat->Stat.StatValue = StoredStat->Stat.LastEMA;
-					bUpdated = true;
-				}
-				else
-				{
-					bUpdated = StoredStat->Stat.StatValue != Stat.StatValue;
-					StoredStat->Stat.StatValue = Stat.StatValue;
-				}
+				FRenderableStat* StoredStat = ApplicationStats.Find(StatName);
 
-				if (bUpdated && StoredStat->Renderable.IsSet())
+				if (StoredStat->Renderable.IsSet())
 				{
-					FText TextToDisplay = FText::FromString(FString::Printf(TEXT("%s: %.*f"), *Stat.StatName.ToString(), Stat.NDecimalPlacesToPrint, StoredStat->Stat.StatValue));
+					FText TextToDisplay = FText::FromString(FString::Printf(TEXT("%s: %s"), *StatName.ToString(), *Stat.ToString()));
 					StoredStat->Renderable.GetValue().Text = TextToDisplay;
 				}
 			}
 			else
 			{
-				FStoredStat StoredStat(Stat);
+				FRenderableStat StoredStat(Stat);
 
 				if (Stat.ShouldDisplayText())
 				{
-					FString			StringToDisplay = FString::Printf(TEXT("%s: %.*f"), *Stat.StatName.ToString(), Stat.NDecimalPlacesToPrint, Stat.StatValue);
+					FString			StringToDisplay = FString::Printf(TEXT("%s: %s"), *StatName.ToString(), *Stat.ToString());
 					FCanvasTextItem Text = CreateText(StringToDisplay, 0, 0);
 					StoredStat.Renderable = Text;
 				}
 
 				ApplicationStats.Add(StatName, StoredStat);
-				bUpdated = true;
 			}
 		}
 
-		if (bUpdated)
+		if (Stat.IsNumeric())
 		{
-			FireStatChanged(FString(TEXT("Application")), Stat.StatName, Stat.StatValue);
+			FireStatChanged(FString(TEXT("Application")), StatName, Stat.GetValue<double>());
 		}
 	}
 
@@ -281,7 +247,7 @@ namespace UE::PixelStreaming2
 
 				for (auto& ApplicationStatEntry : ApplicationStats)
 				{
-					FStoredStat& StatToDraw = ApplicationStatEntry.Value;
+					FRenderableStat& StatToDraw = ApplicationStatEntry.Value;
 					if (!StatToDraw.Renderable.IsSet())
 					{
 						continue;
@@ -336,7 +302,7 @@ namespace UE::PixelStreaming2
 						// Draw the stat value
 						for (auto& StatEntry : StatGroup.GetStoredStats())
 						{
-							FStoredStat& Stat = StatEntry.Value;
+							FRenderableStat& Stat = StatEntry.Value;
 							if (!Stat.Renderable.IsSet())
 							{
 								continue;
@@ -410,13 +376,13 @@ namespace UE::PixelStreaming2
 		double DeltaSeconds = FPlatformTime::ToSeconds64(FPlatformTime::Cycles64() - LastTimeSettingsPolledCycles);
 		if (DeltaSeconds > 1)
 		{
-			StoreApplicationStat(FStatData(FName(TEXT("PixelStreaming2.Encoder.MinQuality")), UPixelStreaming2PluginSettings::CVarEncoderMinQuality.GetValueOnAnyThread(), 0));
-			StoreApplicationStat(FStatData(FName(TEXT("PixelStreaming2.Encoder.MaxQuality")), UPixelStreaming2PluginSettings::CVarEncoderMaxQuality.GetValueOnAnyThread(), 0));
-			StoreApplicationStat(FStatData(FName(TEXT("PixelStreaming2.Encoder.KeyframeInterval (frames)")), UPixelStreaming2PluginSettings::CVarEncoderKeyframeInterval.GetValueOnAnyThread(), 0));
-			StoreApplicationStat(FStatData(FName(TEXT("PixelStreaming2.WebRTC.Fps")), UPixelStreaming2PluginSettings::CVarWebRTCFps.GetValueOnAnyThread(), 0));
-			StoreApplicationStat(FStatData(FName(TEXT("PixelStreaming2.WebRTC.StartBitrate")), UPixelStreaming2PluginSettings::CVarWebRTCStartBitrate.GetValueOnAnyThread(), 0));
-			StoreApplicationStat(FStatData(FName(TEXT("PixelStreaming2.WebRTC.MinBitrate")), UPixelStreaming2PluginSettings::CVarWebRTCMinBitrate.GetValueOnAnyThread(), 0));
-			StoreApplicationStat(FStatData(FName(TEXT("PixelStreaming2.WebRTC.MaxBitrate")), UPixelStreaming2PluginSettings::CVarWebRTCMaxBitrate.GetValueOnAnyThread(), 0));
+			StoreApplicationStat(FStat({ .Name = FName(TEXT("PixelStreaming2.Encoder.MinQuality")) }, UPixelStreaming2PluginSettings::CVarEncoderMinQuality.GetValueOnAnyThread(), 0));
+			StoreApplicationStat(FStat({ .Name = FName(TEXT("PixelStreaming2.Encoder.MaxQuality")) }, UPixelStreaming2PluginSettings::CVarEncoderMaxQuality.GetValueOnAnyThread(), 0));
+			StoreApplicationStat(FStat({ .Name = FName(TEXT("PixelStreaming2.Encoder.KeyframeInterval (frames)")) }, UPixelStreaming2PluginSettings::CVarEncoderKeyframeInterval.GetValueOnAnyThread(), 0));
+			StoreApplicationStat(FStat({ .Name = FName(TEXT("PixelStreaming2.WebRTC.Fps")) }, UPixelStreaming2PluginSettings::CVarWebRTCFps.GetValueOnAnyThread(), 0));
+			StoreApplicationStat(FStat({ .Name = FName(TEXT("PixelStreaming2.WebRTC.StartBitrate")) }, UPixelStreaming2PluginSettings::CVarWebRTCStartBitrate.GetValueOnAnyThread(), 0));
+			StoreApplicationStat(FStat({ .Name = FName(TEXT("PixelStreaming2.WebRTC.MinBitrate")) }, UPixelStreaming2PluginSettings::CVarWebRTCMinBitrate.GetValueOnAnyThread(), 0));
+			StoreApplicationStat(FStat({ .Name = FName(TEXT("PixelStreaming2.WebRTC.MaxBitrate")) }, UPixelStreaming2PluginSettings::CVarWebRTCMaxBitrate.GetValueOnAnyThread(), 0));
 
 			LastTimeSettingsPolledCycles = FPlatformTime::Cycles64();
 		}
@@ -494,21 +460,19 @@ namespace UE::PixelStreaming2
 	// ---------------- FStatGroup ---------------------------
 	// A collection of stats grouped together by a category name
 	//
-	bool FStatGroup::StoreStat(FStatData& StatToStore)
+	bool FStatGroup::StoreStat(FStat& StatToStore)
 	{
-		bool bUpdated = false;
-
-		// If a stat has an alias, use that as the storage key, otherwise use its display name
-		FName& StatName = StatToStore.Alias.IsSet() ? StatToStore.Alias.GetValue() : StatToStore.StatName;
+		// If a stat has an alias, use that as the storage key, otherwise use its underlying name
+		FName StatName = StatToStore.GetDisplayName();
 
 		if (!StoredStats.Contains(StatName))
 		{
-			FStoredStat NewStat(StatToStore);
+			FRenderableStat NewStat(StatToStore);
 
 			// If we are displaying the stat, add a renderable for it
 			if (StatToStore.ShouldDisplayText())
 			{
-				FString			StringToDisplay = FString::Printf(TEXT("%s: %.*f"), *StatToStore.StatName.ToString(), StatToStore.NDecimalPlacesToPrint, StatToStore.StatValue);
+				FString			StringToDisplay = FString::Printf(TEXT("%s: %s"), *StatName.ToString(), *StatToStore.ToString());
 				FCanvasTextItem Text = CreateText(StringToDisplay, 0, 0);
 				NewStat.Renderable = Text;
 			}
@@ -526,44 +490,18 @@ namespace UE::PixelStreaming2
 		else
 		{
 			// We already have this stat, so just update it
-
-			FStoredStat* StoredStat = StoredStats.Find(StatName);
+			FRenderableStat* StoredStat = StoredStats.Find(StatName);
 			if (!StoredStat)
 			{
 				return false;
 			}
 
-			if (StoredStat->Stat.bSmooth && StoredStat->Stat.StatValue != 0)
-			{
-				const int MaxSamples = 60;
-				StoredStat->Stat.NumSamples = FGenericPlatformMath::Min(MaxSamples, StoredStat->Stat.NumSamples + 1);
-				if (StoredStat->Stat.NumSamples < MaxSamples)
-				{
-					StoredStat->Stat.LastEMA = CalcMA(StoredStat->Stat.LastEMA, StoredStat->Stat.NumSamples - 1, StatToStore.StatValue);
-				}
-				else
-				{
-					StoredStat->Stat.LastEMA = CalcEMA(StoredStat->Stat.LastEMA, StoredStat->Stat.NumSamples - 1, StatToStore.StatValue);
-				}
-				StoredStat->Stat.StatValue = StoredStat->Stat.LastEMA;
-				bUpdated = true;
-			}
-			else
-			{
-				bUpdated = StoredStat->Stat.StatValue != StatToStore.StatValue;
-				StoredStat->Stat.StatValue = StatToStore.StatValue;
-			}
-
-			if (!bUpdated)
+			if (StoredStat->ShouldDisplayText() && StoredStat->Renderable.IsSet())
 			{
-				return false;
-			}
-			else if (StoredStat->Stat.ShouldDisplayText() && StoredStat->Renderable.IsSet())
-			{
-				FText TextToDisplay = FText::FromString(FString::Printf(TEXT("%s: %.*f"), *StatToStore.StatName.ToString(), StatToStore.NDecimalPlacesToPrint, StoredStat->Stat.StatValue));
+				FText TextToDisplay = FText::FromString(FString::Printf(TEXT("%s: %s"), *StatName.ToString(), *StatToStore.ToString()));
 				StoredStat->Renderable.GetValue().Text = TextToDisplay;
 			}
-			return bUpdated;
+			return true;
 		}
 	}
 
@@ -572,7 +510,7 @@ namespace UE::PixelStreaming2
 	// Stats specific to a particular peer, as opposed to the entire app.
 	//
 
-	bool FPeerStats::StoreStat(FName StatCategory, FStatData& StatToStore)
+	bool FPeerStats::StoreStat(FName StatCategory, FStat& StatToStore)
 	{
 		if (!StatGroups.Contains(StatCategory))
 		{
@@ -584,19 +522,21 @@ namespace UE::PixelStreaming2
 		return StatGroups[StatCategory].StoreStat(StatToStore);
 	}
 
-	bool UE::PixelStreaming2::FPeerStats::GetStat(FName StatCategory, FName StatToQuery, double& OutValue) const
+	bool UE::PixelStreaming2::FPeerStats::GetStat(FName StatCategory, FName StatToQuery, double& OutValue)
 	{
-		const FStatGroup* Group = GetStatGroups().Find(StatCategory);
+		FStatGroup* Group = StatGroups.Find(StatCategory);
 		if (!Group)
 		{
 			return false;
 		}
-		const FStoredStat* StoredStat = Group->GetStoredStats().Find(StatToQuery);
+
+		FRenderableStat* StoredStat = Group->GetStoredStats().Find(StatToQuery);
 		if (!StoredStat)
 		{
 			return false;
 		}
-		OutValue = StoredStat->Stat.StatValue;
+
+		OutValue = StoredStat->GetValue<double>();
 		return true;
 	}
 
@@ -632,7 +572,7 @@ namespace UE::PixelStreaming2
 	double FStats::AddTimeStat(uint64 Millis, const FString& Label)
 	{
 		const double	DeltaMs = Millis;
-		const FStatData TimeData{ FName(*Label), DeltaMs, 2, true };
+		const FStat TimeData({ .Name = FName(*Label) }, DeltaMs, 2, true);
 		StoreApplicationStat(TimeData);
 		return DeltaMs;
 	}
@@ -642,7 +582,7 @@ namespace UE::PixelStreaming2
 		const uint64	MaxMillis = FGenericPlatformMath::Max(Millis1, Millis2);
 		const uint64	MinMillis = FGenericPlatformMath::Min(Millis1, Millis2);
 		const double	DeltaMs = (MaxMillis - MinMillis) * ((Millis1 > Millis2) ? 1.0 : -1.0);
-		const FStatData TimeData{ FName(*Label), DeltaMs, 2, true };
+		const FStat TimeData({ .Name = FName(*Label) }, DeltaMs, 2, true);
 		StoreApplicationStat(TimeData);
 		return DeltaMs;
 	}
@@ -654,7 +594,7 @@ namespace UE::PixelStreaming2
 		const double TimeGPUDelay = AddTimeStat(FrameMetadata.CaptureProcessGPUDelay, FString::Printf(TEXT("%s Layer %d Frame Capture GPU Delay Time"), *FrameMetadata.ProcessName, FrameMetadata.Layer));
 		const double TimeGPU = AddTimeStat(FrameMetadata.CaptureProcessGPUTime, FString::Printf(TEXT("%s Layer %d Frame Capture GPU Time"), *FrameMetadata.ProcessName, FrameMetadata.Layer));
 
-		const FStatData UseData{ FName(*FString::Printf(TEXT("%s Layer %d Frame Uses"), *FrameMetadata.ProcessName, FrameMetadata.Layer)), static_cast<double>(FrameMetadata.UseCount), 0, false };
+		const FStat UseData({ .Name = FName(*FString::Printf(TEXT("%s Layer %d Frame Uses"), *FrameMetadata.ProcessName, FrameMetadata.Layer)) }, static_cast<double>(FrameMetadata.UseCount));
 		StoreApplicationStat(UseData);
 
 		const int Samples = 100;
diff --git a/Engine/Plugins/Media/PixelStreaming2/Source/PixelStreaming2RTC/Private/Stats.h b/Engine/Plugins/Media/PixelStreaming2/Source/PixelStreaming2RTC/Private/Stats.h
--- a/Engine/Plugins/Media/PixelStreaming2/Source/PixelStreaming2RTC/Private/Stats.h
+++ b/Engine/Plugins/Media/PixelStreaming2/Source/PixelStreaming2RTC/Private/Stats.h
@@ -15,15 +15,14 @@
 namespace UE::PixelStreaming2
 {
 	// Stat that can be optionally rendered
-	struct FStoredStat
+	struct FRenderableStat : public FStat
 	{
-	public:
-		FStatData				   Stat;
-		TOptional<FCanvasTextItem> Renderable;
+		FRenderableStat(FStat& Stat)
+			: FStat(Stat)
+		{
+		}
 
-	public:
-		FStoredStat(FStatData& InStat)
-			: Stat(InStat) {}
+		TOptional<FCanvasTextItem> Renderable;
 	};
 
 	// A grouping of stats by some category name
@@ -38,12 +37,12 @@ namespace UE::PixelStreaming2
 		}
 
 		virtual ~FStatGroup() = default;
-		bool							StoreStat(FStatData& StatToStore);
-		const TMap<FName, FStoredStat>& GetStoredStats() const { return StoredStats; }
-		TMap<FName, FStoredStat>& GetStoredStats() { return StoredStats; }
+		bool						  StoreStat(FStat& StatToStore);
+		TMap<FName, FRenderableStat>& GetStoredStats() { return StoredStats; }
+
 	private:
-		FName					 GroupName;
-		TMap<FName, FStoredStat> StoredStats;
+		FName						 GroupName;
+		TMap<FName, FRenderableStat> StoredStats;
 
 	public:
 		FCanvasTextItem CategoryCanvasItem;
@@ -61,9 +60,8 @@ namespace UE::PixelStreaming2
 			PlayerIdCanvasItem.EnableShadow(FLinearColor::Black);
 		};
 
-		bool						   StoreStat(FName StatCategory, FStatData& StatToStore);
-		bool						   GetStat(FName StatCategory, FName StatToQuery, double& OutValue) const;
-		const TMap<FName, FStatGroup>& GetStatGroups() const { return StatGroups; }
+		bool					 StoreStat(FName StatCategory, FStat& StatToStore);
+		bool					 GetStat(FName StatCategory, FName StatToQuery, double& OutValue);
 		TMap<FName, FStatGroup>& GetStatGroups() { return StatGroups; }
 
 	private:
@@ -89,11 +87,11 @@ namespace UE::PixelStreaming2
 		static FStats*			Get();
 
 		FStats(const FStats&) = delete;
-		bool QueryPeerStat(const FString& PlayerId, FName StatCategory, FName StatToQuery, double& OutValue) const;
+		bool QueryPeerStat(const FString& PlayerId, FName StatCategory, FName StatToQuery, double& OutValue);
 		void RemovePeerStats(const FString& PlayerId);
 		void RemoveAllPeerStats();
-		void StorePeerStat(const FString& PlayerId, FName StatCategory, FStatData Stat);
-		void StoreApplicationStat(FStatData PeerStat);
+		void StorePeerStat(const FString& PlayerId, FName StatCategory, FStat Stat);
+		void StoreApplicationStat(FStat PeerStat);
 		void RemoveAllApplicationStats();
 		void Tick(float DeltaTime);
 
@@ -129,11 +127,11 @@ namespace UE::PixelStreaming2
 	private:
 		static FStats* Instance;
 
-		mutable FCriticalSection  PeerStatsCS;
+		FCriticalSection		  PeerStatsCS;
 		TMap<FString, FPeerStats> PeerStats;
 
-		mutable FCriticalSection ApplicationStatsCS;
-		TMap<FName, FStoredStat> ApplicationStats;
+		FCriticalSection			 ApplicationStatsCS;
+		TMap<FName, FRenderableStat> ApplicationStats;
 
 		int64 LastTimeSettingsPolledCycles = 0;
 
diff --git a/Engine/Plugins/Media/PixelStreaming2/Source/PixelStreaming2RTC/Public/PixelStreaming2StatNames.h b/Engine/Plugins/Media/PixelStreaming2/Source/PixelStreaming2RTC/Public/PixelStreaming2StatNames.h
--- a/Engine/Plugins/Media/PixelStreaming2/Source/PixelStreaming2RTC/Public/PixelStreaming2StatNames.h
+++ b/Engine/Plugins/Media/PixelStreaming2/Source/PixelStreaming2RTC/Public/PixelStreaming2StatNames.h
@@ -71,5 +71,10 @@ namespace PixelStreaming2StatNames
 	PIXELSTREAMING2RTC_API extern const FName DataChannelMessagesSent;
 	PIXELSTREAMING2RTC_API extern const FName DataChannelMessagesReceived;
 	PIXELSTREAMING2RTC_API extern const FName InputController;
+	PIXELSTREAMING2RTC_API extern const FName PacketsSent;
+	PIXELSTREAMING2RTC_API extern const FName EncoderImplementation;
+	PIXELSTREAMING2RTC_API extern const FName MimeType;
+	PIXELSTREAMING2RTC_API extern const FName Channels;
+	PIXELSTREAMING2RTC_API extern const FName ClockRate;
 
 } // namespace PixelStreaming2StatNames
