[PS2] QOL: Allow UseMediaCapture to be toggled at runtime

diff --git a/Engine/Plugins/Media/PixelStreaming2/Source/PixelStreaming2/Internal/VideoCapturer.h b/Engine/Plugins/Media/PixelStreaming2/Source/PixelStreaming2/Internal/VideoCapturer.h
--- a/Engine/Plugins/Media/PixelStreaming2/Source/PixelStreaming2/Internal/VideoCapturer.h
+++ b/Engine/Plugins/Media/PixelStreaming2/Source/PixelStreaming2/Internal/VideoCapturer.h
@@ -60,13 +60,12 @@ namespace UE::PixelStreaming2
 		TSharedPtr<FVideoProducer>					 VideoProducer;
 		TSharedPtr<FPixelCaptureCapturerMultiFormat> FrameCapturer;
 		FDelegateHandle								 CaptureCompleteHandle;
-		FDelegateHandle								 SimulcastEnabledChangedHandle;
-		FDelegateHandle								 CaptureUseFenceChangedHandle;
 		FDelegateHandle								 FramePushedHandle;
 
 		void CreateFrameCapturer();
 		void OnSimulcastEnabledChanged(IConsoleVariable* Var);
 		void OnCaptureUseFenceChanged(IConsoleVariable* Var);
+		void OnUseMediaCaptureChanged(IConsoleVariable* Var);
 		void OnCaptureComplete();
 		void OnFrame(const IPixelCaptureInputFrame& InputFrame);
 	};
diff --git a/Engine/Plugins/Media/PixelStreaming2/Source/PixelStreaming2/Internal/VideoProducer.h b/Engine/Plugins/Media/PixelStreaming2/Source/PixelStreaming2/Internal/VideoProducer.h
--- a/Engine/Plugins/Media/PixelStreaming2/Source/PixelStreaming2/Internal/VideoProducer.h
+++ b/Engine/Plugins/Media/PixelStreaming2/Source/PixelStreaming2/Internal/VideoProducer.h
@@ -8,6 +8,15 @@
 
 namespace UE::PixelStreaming2
 {
+	namespace VideoProducerIdentifiers
+	{
+		PIXELSTREAMING2_API extern const FString FVideoProducer;
+		PIXELSTREAMING2_API extern const FString FVideoProducerBackBuffer;
+		PIXELSTREAMING2_API extern const FString FVideoProducerMediaCapture;
+		PIXELSTREAMING2_API extern const FString FVideoProducerPIEViewport;
+		PIXELSTREAMING2_API extern const FString FVideoProducerRenderTarget;
+	}
+	
 	class PIXELSTREAMING2_API FVideoProducer : public IPixelStreaming2VideoProducer
 	{
 	public:
diff --git a/Engine/Plugins/Media/PixelStreaming2/Source/PixelStreaming2/Private/VideoCapturer.cpp b/Engine/Plugins/Media/PixelStreaming2/Source/PixelStreaming2/Private/VideoCapturer.cpp
--- a/Engine/Plugins/Media/PixelStreaming2/Source/PixelStreaming2/Private/VideoCapturer.cpp
+++ b/Engine/Plugins/Media/PixelStreaming2/Source/PixelStreaming2/Private/VideoCapturer.cpp
@@ -29,8 +29,9 @@ namespace UE::PixelStreaming2
 
 		if (UPixelStreaming2PluginSettings::FDelegates* Delegates = UPixelStreaming2PluginSettings::Delegates())
 		{
-			VideoCapturer->SimulcastEnabledChangedHandle = Delegates->OnSimulcastEnabledChanged.AddSP(VideoCapturer.ToSharedRef(), &FVideoCapturer::OnSimulcastEnabledChanged);
-			VideoCapturer->CaptureUseFenceChangedHandle = Delegates->OnCaptureUseFenceChanged.AddSP(VideoCapturer.ToSharedRef(), &FVideoCapturer::OnCaptureUseFenceChanged);
+			Delegates->OnSimulcastEnabledChanged.AddSP(VideoCapturer.ToSharedRef(), &FVideoCapturer::OnSimulcastEnabledChanged);
+			Delegates->OnCaptureUseFenceChanged.AddSP(VideoCapturer.ToSharedRef(), &FVideoCapturer::OnCaptureUseFenceChanged);
+			Delegates->OnUseMediaCaptureChanged.AddSP(VideoCapturer.ToSharedRef(), &FVideoCapturer::OnUseMediaCaptureChanged);
 		}
 		return VideoCapturer;
 	}
@@ -113,6 +114,11 @@ namespace UE::PixelStreaming2
 		CreateFrameCapturer();
 	}
 
+	void FVideoCapturer::OnUseMediaCaptureChanged(IConsoleVariable* Var)
+	{
+		CreateFrameCapturer();
+	}
+
 	void FVideoCapturer::CreateFrameCapturer()
 	{
 		if (FrameCapturer != nullptr)
diff --git a/Engine/Plugins/Media/PixelStreaming2/Source/PixelStreaming2/Private/VideoProducer.cpp b/Engine/Plugins/Media/PixelStreaming2/Source/PixelStreaming2/Private/VideoProducer.cpp
--- a/Engine/Plugins/Media/PixelStreaming2/Source/PixelStreaming2/Private/VideoProducer.cpp
+++ b/Engine/Plugins/Media/PixelStreaming2/Source/PixelStreaming2/Private/VideoProducer.cpp
@@ -4,6 +4,14 @@
 
 namespace UE::PixelStreaming2
 {
+	namespace VideoProducerIdentifiers
+	{
+		const FString FVideoProducer = TEXT("The default video producer - override me");
+		const FString FVideoProducerBackBuffer = TEXT("the Back Buffer");
+		const FString FVideoProducerMediaCapture = TEXT("a Media Capture Video Input");
+		const FString FVideoProducerPIEViewport = TEXT("the PIE Viewport");
+		const FString FVideoProducerRenderTarget = TEXT("a Render Target");
+	}
 
 	TSharedPtr<FVideoProducer> FVideoProducer::Create()
 	{
@@ -17,7 +25,7 @@ namespace UE::PixelStreaming2
 
 	FString FVideoProducer::ToString()
 	{
-		return TEXT("The default video producer - override me");
+		return VideoProducerIdentifiers::FVideoProducer;
 	}
 
 } // namespace UE::PixelStreaming2
\ No newline at end of file
diff --git a/Engine/Plugins/Media/PixelStreaming2/Source/PixelStreaming2/Private/VideoProducerBackBuffer.cpp b/Engine/Plugins/Media/PixelStreaming2/Source/PixelStreaming2/Private/VideoProducerBackBuffer.cpp
--- a/Engine/Plugins/Media/PixelStreaming2/Source/PixelStreaming2/Private/VideoProducerBackBuffer.cpp
+++ b/Engine/Plugins/Media/PixelStreaming2/Source/PixelStreaming2/Private/VideoProducerBackBuffer.cpp
@@ -54,7 +54,7 @@ namespace UE::PixelStreaming2
 
 	FString FVideoProducerBackBuffer::ToString()
 	{
-		return TEXT("the Back Buffer");
+		return VideoProducerIdentifiers::FVideoProducerBackBuffer;
 	}
 
 } // namespace UE::PixelStreaming2
\ No newline at end of file
diff --git a/Engine/Plugins/Media/PixelStreaming2/Source/PixelStreaming2/Private/VideoProducerMediaCapture.cpp b/Engine/Plugins/Media/PixelStreaming2/Source/PixelStreaming2/Private/VideoProducerMediaCapture.cpp
--- a/Engine/Plugins/Media/PixelStreaming2/Source/PixelStreaming2/Private/VideoProducerMediaCapture.cpp
+++ b/Engine/Plugins/Media/PixelStreaming2/Source/PixelStreaming2/Private/VideoProducerMediaCapture.cpp
@@ -83,7 +83,7 @@ namespace UE::PixelStreaming2
 
 	FString FVideoProducerMediaCapture::ToString()
 	{
-		return TEXT("a Media Capture Video Input");
+		return VideoProducerIdentifiers::FVideoProducerMediaCapture;
 	}
 
 	void FVideoProducerMediaCapture::OnCaptureActiveViewportStateChanged()
diff --git a/Engine/Plugins/Media/PixelStreaming2/Source/PixelStreaming2/Private/VideoProducerPIEViewport.cpp b/Engine/Plugins/Media/PixelStreaming2/Source/PixelStreaming2/Private/VideoProducerPIEViewport.cpp
--- a/Engine/Plugins/Media/PixelStreaming2/Source/PixelStreaming2/Private/VideoProducerPIEViewport.cpp
+++ b/Engine/Plugins/Media/PixelStreaming2/Source/PixelStreaming2/Private/VideoProducerPIEViewport.cpp
@@ -55,7 +55,7 @@ namespace UE::PixelStreaming2
 
 	FString FVideoProducerPIEViewport::ToString()
 	{
-		return TEXT("the PIE Viewport");
+		return VideoProducerIdentifiers::FVideoProducerPIEViewport;
 	}
 
 } // namespace UE::PixelStreaming2
\ No newline at end of file
diff --git a/Engine/Plugins/Media/PixelStreaming2/Source/PixelStreaming2/Private/VideoProducerRenderTarget.cpp b/Engine/Plugins/Media/PixelStreaming2/Source/PixelStreaming2/Private/VideoProducerRenderTarget.cpp
--- a/Engine/Plugins/Media/PixelStreaming2/Source/PixelStreaming2/Private/VideoProducerRenderTarget.cpp
+++ b/Engine/Plugins/Media/PixelStreaming2/Source/PixelStreaming2/Private/VideoProducerRenderTarget.cpp
@@ -40,7 +40,7 @@ namespace UE::PixelStreaming2
 
 	FString FVideoProducerRenderTarget::ToString()
 	{
-		return TEXT("a Render Target");
+		return VideoProducerIdentifiers::FVideoProducerRenderTarget;
 	}
 
 } // namespace UE::PixelStreaming2
\ No newline at end of file
diff --git a/Engine/Plugins/Media/PixelStreaming2/Source/PixelStreaming2RTC/Private/EpicRtcStreamer.cpp b/Engine/Plugins/Media/PixelStreaming2/Source/PixelStreaming2RTC/Private/EpicRtcStreamer.cpp
--- a/Engine/Plugins/Media/PixelStreaming2/Source/PixelStreaming2RTC/Private/EpicRtcStreamer.cpp
+++ b/Engine/Plugins/Media/PixelStreaming2/Source/PixelStreaming2RTC/Private/EpicRtcStreamer.cpp
@@ -190,6 +190,11 @@ namespace UE::PixelStreaming2
 
 		VideoCapturer->OnFrameCapturerCreated.AddSP(Streamer.ToSharedRef(), &FEpicRtcStreamer::OnFrameCapturerCreated);
 
+		if (UPixelStreaming2PluginSettings::FDelegates* Delegates = UPixelStreaming2PluginSettings::Delegates())
+		{
+			Delegates->OnUseMediaCaptureChanged.AddSP(Streamer.ToSharedRef(), &FEpicRtcStreamer::OnUseMediaCaptureChanged);
+		}
+
 		FPixelStreaming2RTCModule::GetModule()->GetStatsCollector()->OnStatsReady.AddSP(Streamer.ToSharedRef(), &FEpicRtcStreamer::OnStatsReady);
 	}
 
@@ -406,6 +411,32 @@ namespace UE::PixelStreaming2
 		SendAllPlayersMessage(MessageName, Descriptor);
 	}
 
+	void FEpicRtcStreamer::OnUseMediaCaptureChanged(IConsoleVariable* CVar)
+	{
+		TSharedPtr<IPixelStreaming2VideoProducer> VideoProducer = VideoCapturer->GetVideoProducer();
+		if (!VideoProducer)
+		{
+			return;
+		}
+
+		FString Identifier = VideoProducer->ToString();
+		// If the identifier is not one of these two types, it means it was user created so we don't want to
+		// override it
+		if (Identifier != VideoProducerIdentifiers::FVideoProducerMediaCapture && Identifier != VideoProducerIdentifiers::FVideoProducerBackBuffer)
+		{
+			return;
+		}
+
+		if (CVar->GetBool())
+		{
+			SetVideoProducer(FVideoProducerMediaCapture::CreateActiveViewportCapture());
+		}
+		else
+		{
+			SetVideoProducer(FVideoProducerBackBuffer::Create());
+		}
+	}
+
 	EpicRtcBitrate FEpicRtcStreamer::GetBitrates() const
 	{
 		const int MinBitrate = UPixelStreaming2PluginSettings::CVarWebRTCMinBitrate.GetValueOnAnyThread();
@@ -1148,7 +1179,7 @@ namespace UE::PixelStreaming2
 			for (int i = 0; i < SimulcastParams.Num(); ++i)
 			{
 				const FPixelStreaming2SimulcastLayer& SpatialLayer = SimulcastParams[i];
-				const bool bIsSimulcast = SimulcastParams.Num() > 1;
+				const bool							  bIsSimulcast = SimulcastParams.Num() > 1;
 
 				if (bIsSimulcast)
 				{
diff --git a/Engine/Plugins/Media/PixelStreaming2/Source/PixelStreaming2RTC/Private/EpicRtcStreamer.h b/Engine/Plugins/Media/PixelStreaming2/Source/PixelStreaming2RTC/Private/EpicRtcStreamer.h
--- a/Engine/Plugins/Media/PixelStreaming2/Source/PixelStreaming2RTC/Private/EpicRtcStreamer.h
+++ b/Engine/Plugins/Media/PixelStreaming2/Source/PixelStreaming2RTC/Private/EpicRtcStreamer.h
@@ -100,6 +100,7 @@ namespace UE::PixelStreaming2
 		void OnFrameCapturerCreated();
 		void OnUIInteraction(FMemoryReader Ar);
 		void OnSendMessage(FString MessageName, FMemoryReader Ar);
+		void OnUseMediaCaptureChanged(IConsoleVariable* CVar);
 		EpicRtcBitrate GetBitrates() const;
 
 	private:
diff --git a/Engine/Plugins/Media/PixelStreaming2/Source/PixelStreaming2RTC/Private/EpicRtcVideoCapturer.cpp b/Engine/Plugins/Media/PixelStreaming2/Source/PixelStreaming2RTC/Private/EpicRtcVideoCapturer.cpp
--- a/Engine/Plugins/Media/PixelStreaming2/Source/PixelStreaming2RTC/Private/EpicRtcVideoCapturer.cpp
+++ b/Engine/Plugins/Media/PixelStreaming2/Source/PixelStreaming2RTC/Private/EpicRtcVideoCapturer.cpp
@@ -17,8 +17,9 @@ namespace UE::PixelStreaming2
 
 		if (UPixelStreaming2PluginSettings::FDelegates* Delegates = UPixelStreaming2PluginSettings::Delegates())
 		{
-			VideoCapturer->SimulcastEnabledChangedHandle = Delegates->OnSimulcastEnabledChanged.AddSP(VideoCapturer.ToSharedRef(), &FEpicRtcVideoCapturer::OnSimulcastEnabledChanged);
-			VideoCapturer->CaptureUseFenceChangedHandle = Delegates->OnCaptureUseFenceChanged.AddSP(VideoCapturer.ToSharedRef(), &FEpicRtcVideoCapturer::OnCaptureUseFenceChanged);
+			Delegates->OnSimulcastEnabledChanged.AddSP(VideoCapturer.ToSharedRef(), &FEpicRtcVideoCapturer::OnSimulcastEnabledChanged);
+			Delegates->OnCaptureUseFenceChanged.AddSP(VideoCapturer.ToSharedRef(), &FEpicRtcVideoCapturer::OnCaptureUseFenceChanged);
+			Delegates->OnUseMediaCaptureChanged.AddSP(VideoCapturer.ToSharedRef(), &FEpicRtcVideoCapturer::OnUseMediaCaptureChanged);
 		}
 		return VideoCapturer;
 	}
@@ -33,4 +34,4 @@ namespace UE::PixelStreaming2
 	{
 		return new FEpicRtcVideoBufferMultiFormatLayered(FrameCapturer);
 	}
-} // namespace UE::PixelStreaming2RTC
\ No newline at end of file
+} // namespace UE::PixelStreaming2
\ No newline at end of file
diff --git a/Engine/Plugins/Media/PixelStreaming2/Source/PixelStreaming2Settings/Internal/PixelStreaming2PluginSettings.h b/Engine/Plugins/Media/PixelStreaming2/Source/PixelStreaming2Settings/Internal/PixelStreaming2PluginSettings.h
--- a/Engine/Plugins/Media/PixelStreaming2/Source/PixelStreaming2Settings/Internal/PixelStreaming2PluginSettings.h
+++ b/Engine/Plugins/Media/PixelStreaming2/Source/PixelStreaming2Settings/Internal/PixelStreaming2PluginSettings.h
@@ -595,6 +595,9 @@ public:
 		DECLARE_TS_MULTICAST_DELEGATE_OneParam(FOnCaptureUseFenceChanged, IConsoleVariable*);
 		FOnCaptureUseFenceChanged OnCaptureUseFenceChanged;
 
+		DECLARE_TS_MULTICAST_DELEGATE_OneParam(FOnUseMediaCaptureChanged, IConsoleVariable*);
+		FOnUseMediaCaptureChanged OnUseMediaCaptureChanged;
+
 		DECLARE_TS_MULTICAST_DELEGATE_OneParam(FOnWebRTCFpsChanged, IConsoleVariable*);
 		FOnWebRTCFpsChanged OnWebRTCFpsChanged;
 
diff --git a/Engine/Plugins/Media/PixelStreaming2/Source/PixelStreaming2Settings/Private/PixelStreaming2PluginSettings.cpp b/Engine/Plugins/Media/PixelStreaming2/Source/PixelStreaming2Settings/Private/PixelStreaming2PluginSettings.cpp
--- a/Engine/Plugins/Media/PixelStreaming2/Source/PixelStreaming2Settings/Private/PixelStreaming2PluginSettings.cpp
+++ b/Engine/Plugins/Media/PixelStreaming2/Source/PixelStreaming2Settings/Private/PixelStreaming2PluginSettings.cpp
@@ -299,6 +299,7 @@ TAutoConsoleVariable<bool> UPixelStreaming2PluginSettings::CVarUseMediaCapture(
 	TEXT("PixelStreaming2.UseMediaCapture"),
 	true,
 	TEXT("Use Media Capture from MediaIOFramework to capture frames rather than Pixel Streamings internal backbuffer sources."),
+	FConsoleVariableDelegate::CreateLambda([](IConsoleVariable* Var) { Delegates()->OnUseMediaCaptureChanged.Broadcast(Var); }),
 	ECVF_Default);
 
 TAutoConsoleVariable<FString> UPixelStreaming2PluginSettings::CVarDefaultStreamerID(
