[PixelCapture, PS2] Add: UserData member on PixelCapture Metadata. PixelStreaming uses this to pass information about the frame production through the PixelCapture pipeline for display with `stat PixelStreaming2`.

diff --git a/Engine/Plugins/Media/PixelCapture/Source/PixelCapture/Private/PixelCaptureInputFrameI420.cpp b/Engine/Plugins/Media/PixelCapture/Source/PixelCapture/Private/PixelCaptureInputFrameI420.cpp
--- a/Engine/Plugins/Media/PixelCapture/Source/PixelCapture/Private/PixelCaptureInputFrameI420.cpp
+++ b/Engine/Plugins/Media/PixelCapture/Source/PixelCapture/Private/PixelCaptureInputFrameI420.cpp
@@ -3,8 +3,9 @@
 #include "PixelCaptureInputFrameI420.h"
 #include "PixelCaptureBufferFormat.h"
 
-FPixelCaptureInputFrameI420::FPixelCaptureInputFrameI420(TSharedPtr<FPixelCaptureBufferI420> Buffer)
-	: I420Buffer(Buffer)
+FPixelCaptureInputFrameI420::FPixelCaptureInputFrameI420(TSharedPtr<FPixelCaptureBufferI420> Buffer, TSharedPtr<FPixelCaptureUserData> UserData)
+	: IPixelCaptureInputFrame(UserData)
+	, I420Buffer(Buffer)
 {
 	Metadata.SourceTime = rtc::TimeMillis();
 }
diff --git a/Engine/Plugins/Media/PixelCapture/Source/PixelCapture/Private/PixelCaptureInputFrameNV12.cpp b/Engine/Plugins/Media/PixelCapture/Source/PixelCapture/Private/PixelCaptureInputFrameNV12.cpp
--- a/Engine/Plugins/Media/PixelCapture/Source/PixelCapture/Private/PixelCaptureInputFrameNV12.cpp
+++ b/Engine/Plugins/Media/PixelCapture/Source/PixelCapture/Private/PixelCaptureInputFrameNV12.cpp
@@ -3,8 +3,9 @@
 #include "PixelCaptureInputFrameNV12.h"
 #include "PixelCaptureBufferFormat.h"
 
-FPixelCaptureInputFrameNV12::FPixelCaptureInputFrameNV12(TSharedPtr<FPixelCaptureBufferNV12> Buffer)
-	: NV12Buffer(Buffer)
+FPixelCaptureInputFrameNV12::FPixelCaptureInputFrameNV12(TSharedPtr<FPixelCaptureBufferNV12> Buffer, TSharedPtr<FPixelCaptureUserData> UserData)
+	: IPixelCaptureInputFrame(UserData)
+	, NV12Buffer(Buffer)
 {
 	Metadata.SourceTime = rtc::TimeMillis();
 }
diff --git a/Engine/Plugins/Media/PixelCapture/Source/PixelCapture/Private/PixelCaptureInputFrameRHI.cpp b/Engine/Plugins/Media/PixelCapture/Source/PixelCapture/Private/PixelCaptureInputFrameRHI.cpp
--- a/Engine/Plugins/Media/PixelCapture/Source/PixelCapture/Private/PixelCaptureInputFrameRHI.cpp
+++ b/Engine/Plugins/Media/PixelCapture/Source/PixelCapture/Private/PixelCaptureInputFrameRHI.cpp
@@ -3,8 +3,9 @@
 #include "PixelCaptureInputFrameRHI.h"
 #include "PixelCaptureBufferFormat.h"
 
-FPixelCaptureInputFrameRHI::FPixelCaptureInputFrameRHI(FTextureRHIRef InFrameTexture)
-	: FrameTexture(InFrameTexture)
+FPixelCaptureInputFrameRHI::FPixelCaptureInputFrameRHI(FTextureRHIRef InFrameTexture, TSharedPtr<FPixelCaptureUserData> UserData)
+	: IPixelCaptureInputFrame(UserData)
+	, FrameTexture(InFrameTexture)
 {
 	Metadata.SourceTime = rtc::TimeMillis();
 }
diff --git a/Engine/Plugins/Media/PixelCapture/Source/PixelCapture/Public/IPixelCaptureInputFrame.h b/Engine/Plugins/Media/PixelCapture/Source/PixelCapture/Public/IPixelCaptureInputFrame.h
--- a/Engine/Plugins/Media/PixelCapture/Source/PixelCapture/Public/IPixelCaptureInputFrame.h
+++ b/Engine/Plugins/Media/PixelCapture/Source/PixelCapture/Public/IPixelCaptureInputFrame.h
@@ -15,6 +15,9 @@ class PIXELCAPTURE_API IPixelCaptureInputFrame
 {
 public:
 	IPixelCaptureInputFrame() { Metadata.SourceTime = rtc::TimeMillis(); }
+
+	IPixelCaptureInputFrame(TSharedPtr<FPixelCaptureUserData> UserData) { Metadata.SourceTime = rtc::TimeMillis(); Metadata.UserData = UserData; }
+
 	virtual ~IPixelCaptureInputFrame() = default;
 
 	/**
diff --git a/Engine/Plugins/Media/PixelCapture/Source/PixelCapture/Public/PixelCaptureFrameMetadata.h b/Engine/Plugins/Media/PixelCapture/Source/PixelCapture/Public/PixelCaptureFrameMetadata.h
--- a/Engine/Plugins/Media/PixelCapture/Source/PixelCapture/Public/PixelCaptureFrameMetadata.h
+++ b/Engine/Plugins/Media/PixelCapture/Source/PixelCapture/Public/PixelCaptureFrameMetadata.h
@@ -4,6 +4,10 @@
 
 #include "CoreMinimal.h"
 
+class FPixelCaptureUserData
+{
+};
+
 struct FPixelCaptureFrameMetadata
 {
 	// Identifier for the capture pipeline/process this frame took
@@ -45,6 +49,9 @@ struct FPixelCaptureFrameMetadata
 	uint64 LastDecodeStartTime = 0;
 	uint64 LastDecodeEndTime = 0;
 
+	// Pointer to a user class that stores additional metadata that they want to keep throughout the process
+	TSharedPtr<FPixelCaptureUserData> UserData;
+
 	// wanted this to be explicit with a name
 	FPixelCaptureFrameMetadata Copy() const { return *this; }
 };
diff --git a/Engine/Plugins/Media/PixelCapture/Source/PixelCapture/Public/PixelCaptureInputFrameI420.h b/Engine/Plugins/Media/PixelCapture/Source/PixelCapture/Public/PixelCaptureInputFrameI420.h
--- a/Engine/Plugins/Media/PixelCapture/Source/PixelCapture/Public/PixelCaptureInputFrameI420.h
+++ b/Engine/Plugins/Media/PixelCapture/Source/PixelCapture/Public/PixelCaptureInputFrameI420.h
@@ -11,7 +11,7 @@
 class PIXELCAPTURE_API FPixelCaptureInputFrameI420 : public IPixelCaptureInputFrame
 {
 public:
-	FPixelCaptureInputFrameI420(TSharedPtr<FPixelCaptureBufferI420> Buffer);
+	FPixelCaptureInputFrameI420(TSharedPtr<FPixelCaptureBufferI420> Buffer, TSharedPtr<FPixelCaptureUserData> UserData = nullptr);
 	virtual ~FPixelCaptureInputFrameI420() = default;
 
 	virtual int32 GetType() const override;
diff --git a/Engine/Plugins/Media/PixelCapture/Source/PixelCapture/Public/PixelCaptureInputFrameNV12.h b/Engine/Plugins/Media/PixelCapture/Source/PixelCapture/Public/PixelCaptureInputFrameNV12.h
--- a/Engine/Plugins/Media/PixelCapture/Source/PixelCapture/Public/PixelCaptureInputFrameNV12.h
+++ b/Engine/Plugins/Media/PixelCapture/Source/PixelCapture/Public/PixelCaptureInputFrameNV12.h
@@ -11,7 +11,7 @@
 class PIXELCAPTURE_API FPixelCaptureInputFrameNV12 : public IPixelCaptureInputFrame
 {
 public:
-	FPixelCaptureInputFrameNV12(TSharedPtr<FPixelCaptureBufferNV12> Buffer);
+	FPixelCaptureInputFrameNV12(TSharedPtr<FPixelCaptureBufferNV12> Buffer, TSharedPtr<FPixelCaptureUserData> UserData = nullptr);
 	virtual ~FPixelCaptureInputFrameNV12() = default;
 
 	virtual int32 GetType() const override;
diff --git a/Engine/Plugins/Media/PixelCapture/Source/PixelCapture/Public/PixelCaptureInputFrameRHI.h b/Engine/Plugins/Media/PixelCapture/Source/PixelCapture/Public/PixelCaptureInputFrameRHI.h
--- a/Engine/Plugins/Media/PixelCapture/Source/PixelCapture/Public/PixelCaptureInputFrameRHI.h
+++ b/Engine/Plugins/Media/PixelCapture/Source/PixelCapture/Public/PixelCaptureInputFrameRHI.h
@@ -11,7 +11,7 @@
 class PIXELCAPTURE_API FPixelCaptureInputFrameRHI : public IPixelCaptureInputFrame
 {
 public:
-	FPixelCaptureInputFrameRHI(FTextureRHIRef InFrameTexture);
+	FPixelCaptureInputFrameRHI(FTextureRHIRef InFrameTexture, TSharedPtr<FPixelCaptureUserData> UserData = nullptr);
 	virtual ~FPixelCaptureInputFrameRHI() = default;
 
 	virtual int32 GetType() const override;
diff --git a/Engine/Plugins/Media/PixelStreaming2/Source/PixelStreaming2/Internal/PixelStreaming2MediaIOCapture.h b/Engine/Plugins/Media/PixelStreaming2/Source/PixelStreaming2/Internal/PixelStreaming2MediaIOCapture.h
--- a/Engine/Plugins/Media/PixelStreaming2/Source/PixelStreaming2/Internal/PixelStreaming2MediaIOCapture.h
+++ b/Engine/Plugins/Media/PixelStreaming2/Source/PixelStreaming2/Internal/PixelStreaming2MediaIOCapture.h
@@ -45,6 +45,7 @@ public:
 	// We override the texture flags because on Mac we want the texture to have the CPU_Readback flag
 	virtual ETextureCreateFlags GetOutputTextureFlags() const override;
 	virtual void WaitForGPU(FRHITexture* InRHITexture) override;
+	virtual TSharedPtr<FMediaCaptureUserData, ESPMode::ThreadSafe> GetCaptureFrameUserData_GameThread() override;
 	//~ End UMediaCapture interface
 
 	TSharedPtr<FSceneViewport>				GetViewport() const { return SceneViewport.Pin(); }
@@ -57,7 +58,7 @@ public:
 	FOnCaptureViewportInitialized OnCaptureViewportInitialized;
 
 private:
-	void HandleCapturedFrame(FTextureRHIRef InTexture);
+	void HandleCapturedFrame(FTextureRHIRef InTexture, TSharedPtr<FMediaCaptureUserData, ESPMode::ThreadSafe> InUserData);
 	void UpdateCaptureResolution(int32 Width, int32 Height);
 
 private:
diff --git a/Engine/Plugins/Media/PixelStreaming2/Source/PixelStreaming2/Internal/VideoProducer.h b/Engine/Plugins/Media/PixelStreaming2/Source/PixelStreaming2/Internal/VideoProducer.h
--- a/Engine/Plugins/Media/PixelStreaming2/Source/PixelStreaming2/Internal/VideoProducer.h
+++ b/Engine/Plugins/Media/PixelStreaming2/Source/PixelStreaming2/Internal/VideoProducer.h
@@ -4,6 +4,7 @@
 
 #include "Delegates/DelegateCombinations.h"
 #include "IPixelStreaming2VideoProducer.h"
+#include "MediaCapture.h"
 #include "Templates/SharedPointer.h"
 
 namespace UE::PixelStreaming2
@@ -15,8 +16,18 @@ namespace UE::PixelStreaming2
 		PIXELSTREAMING2_API extern const FString FVideoProducerMediaCapture;
 		PIXELSTREAMING2_API extern const FString FVideoProducerPIEViewport;
 		PIXELSTREAMING2_API extern const FString FVideoProducerRenderTarget;
-	}
-	
+	} // namespace VideoProducerIdentifiers
+
+	class PIXELSTREAMING2_API FVideoProducerUserData : public FMediaCaptureUserData, public FPixelCaptureUserData
+	{
+	public:
+		FVideoProducerUserData() = default;
+
+		uint64	ProductionBeginCycles = 0;
+		uint64	ProductionEndCycles = 0;
+		FString ProducerName = TEXT("");
+	};
+
 	class PIXELSTREAMING2_API FVideoProducer : public IPixelStreaming2VideoProducer
 	{
 	public:
diff --git a/Engine/Plugins/Media/PixelStreaming2/Source/PixelStreaming2/Private/PixelStreaming2MediaIOCapture.cpp b/Engine/Plugins/Media/PixelStreaming2/Source/PixelStreaming2/Private/PixelStreaming2MediaIOCapture.cpp
--- a/Engine/Plugins/Media/PixelStreaming2/Source/PixelStreaming2/Private/PixelStreaming2MediaIOCapture.cpp
+++ b/Engine/Plugins/Media/PixelStreaming2/Source/PixelStreaming2/Private/PixelStreaming2MediaIOCapture.cpp
@@ -17,7 +17,7 @@ void UPixelStreaming2MediaIOCapture::OnRHIResourceCaptured_RenderingThread(
 	TSharedPtr<FMediaCaptureUserData, ESPMode::ThreadSafe> InUserData,
 	FTextureRHIRef										   InTexture)
 {
-	HandleCapturedFrame(InTexture);
+	HandleCapturedFrame(InTexture, InUserData);
 }
 
 void UPixelStreaming2MediaIOCapture::OnCustomCapture_RenderingThread(
@@ -169,6 +169,15 @@ void UPixelStreaming2MediaIOCapture::WaitForGPU(FRHITexture* InRHITexture)
 	}
 }
 
+TSharedPtr<FMediaCaptureUserData, ESPMode::ThreadSafe> UPixelStreaming2MediaIOCapture::GetCaptureFrameUserData_GameThread()
+{
+	TSharedPtr<UE::PixelStreaming2::FVideoProducerUserData> UserData = MakeShared<UE::PixelStreaming2::FVideoProducerUserData>();
+	UserData->ProducerName = TEXT("FVideoProducerMediaCapture");
+	UserData->ProductionBeginCycles = FPlatformTime::Cycles64();
+
+	return UserData;
+}
+
 bool UPixelStreaming2MediaIOCapture::PostInitializeCaptureViewport(TSharedPtr<FSceneViewport>& InSceneViewport)
 {
 	SceneViewport = TWeakPtr<FSceneViewport>(InSceneViewport);
@@ -202,13 +211,15 @@ void UPixelStreaming2MediaIOCapture::ViewportResized(FViewport* Viewport, uint32
 	}
 }
 
-void UPixelStreaming2MediaIOCapture::HandleCapturedFrame(FTextureRHIRef InTexture)
+void UPixelStreaming2MediaIOCapture::HandleCapturedFrame(FTextureRHIRef InTexture, TSharedPtr<FMediaCaptureUserData, ESPMode::ThreadSafe> InUserData)
 {
 	TSharedPtr<IPixelStreaming2VideoProducer> VideoProducerPtr = VideoProducer.Pin();
 	if (VideoProducerPtr)
 	{
+		TSharedPtr<UE::PixelStreaming2::FVideoProducerUserData> UserData = StaticCastSharedPtr<UE::PixelStreaming2::FVideoProducerUserData>(InUserData);
+		UserData->ProductionEndCycles = FPlatformTime::Cycles64();
 		UpdateCaptureResolution(InTexture->GetDesc().Extent.X, InTexture->GetDesc().Extent.Y);
-		VideoProducerPtr->PushFrame(FPixelCaptureInputFrameRHI(InTexture));
+		VideoProducerPtr->PushFrame(FPixelCaptureInputFrameRHI(InTexture, UserData));
 	}
 }
 
diff --git a/Engine/Plugins/Media/PixelStreaming2/Source/PixelStreaming2/Private/VideoProducerPIEViewport.cpp b/Engine/Plugins/Media/PixelStreaming2/Source/PixelStreaming2/Private/VideoProducerPIEViewport.cpp
--- a/Engine/Plugins/Media/PixelStreaming2/Source/PixelStreaming2/Private/VideoProducerPIEViewport.cpp
+++ b/Engine/Plugins/Media/PixelStreaming2/Source/PixelStreaming2/Private/VideoProducerPIEViewport.cpp
@@ -47,9 +47,14 @@ namespace UE::PixelStreaming2
 			return;
 		}
 
+		TSharedPtr<FVideoProducerUserData> UserData = MakeShared<FVideoProducerUserData>();
+		UserData->ProducerName = TEXT("FVideoProducerPIEViewport");
+		UserData->ProductionBeginCycles = FPlatformTime::Cycles64();
+
 		ENQUEUE_RENDER_COMMAND(StreamViewportTextureCommand)
-		([&, FrameBuffer](FRHICommandList& RHICmdList) {
-			PushFrame(FPixelCaptureInputFrameRHI(FrameBuffer));
+		([&, FrameBuffer, UserData](FRHICommandList& RHICmdList) {
+			UserData->ProductionEndCycles = FPlatformTime::Cycles64();
+			PushFrame(FPixelCaptureInputFrameRHI(FrameBuffer, UserData));
 		});
 	}
 
diff --git a/Engine/Plugins/Media/PixelStreaming2/Source/PixelStreaming2Editor/Internal/VideoProducerBackBufferComposited.h b/Engine/Plugins/Media/PixelStreaming2/Source/PixelStreaming2Editor/Internal/VideoProducerBackBufferComposited.h
--- a/Engine/Plugins/Media/PixelStreaming2/Source/PixelStreaming2Editor/Internal/VideoProducerBackBufferComposited.h
+++ b/Engine/Plugins/Media/PixelStreaming2/Source/PixelStreaming2Editor/Internal/VideoProducerBackBufferComposited.h
@@ -23,6 +23,8 @@ namespace UE::EditorPixelStreaming2
 
 		virtual FString ToString() override;
 
+		virtual bool IsFrameAlreadyCopied() override { return true; }
+
 		DECLARE_MULTICAST_DELEGATE_OneParam(FOnFrameSizeChanged, TWeakPtr<FIntRect>);
 		FOnFrameSizeChanged OnFrameSizeChanged;
 
@@ -56,7 +58,7 @@ namespace UE::EditorPixelStreaming2
 
 	private:
 		FVideoProducerBackBufferComposited();
-		void CompositeWindows();
+		void CompositeWindows(TSharedPtr<FVideoProducerUserData> UserData);
 
 		void OnBackBufferReady(SWindow& SlateWindow, const FTextureRHIRef& FrameBuffer);
 		void OnPreTick(float DeltaTime);
diff --git a/Engine/Plugins/Media/PixelStreaming2/Source/PixelStreaming2Editor/Private/VideoProducerBackBufferComposited.cpp b/Engine/Plugins/Media/PixelStreaming2/Source/PixelStreaming2Editor/Private/VideoProducerBackBufferComposited.cpp
--- a/Engine/Plugins/Media/PixelStreaming2/Source/PixelStreaming2Editor/Private/VideoProducerBackBufferComposited.cpp
+++ b/Engine/Plugins/Media/PixelStreaming2/Source/PixelStreaming2Editor/Private/VideoProducerBackBufferComposited.cpp
@@ -124,6 +124,10 @@ namespace UE::EditorPixelStreaming2
 		 */
 		UE_LOG(LogPixelStreaming2BackBufferComposited, Verbose, TEXT("Type: %s"), *SlateWindow.GetTitle().ToString());
 
+		TSharedPtr<FVideoProducerUserData> UserData = MakeShared<FVideoProducerUserData>();
+		UserData->ProducerName = TEXT("FVideoProducerBackBufferComposited");
+		UserData->ProductionBeginCycles = FPlatformTime::Cycles64();
+
 		{
 			FScopeLock Lock(&TopLevelWindowsCriticalSection);
 			if (TopLevelWindows.IsEmpty())
@@ -163,13 +167,13 @@ namespace UE::EditorPixelStreaming2
 				if (!bSkipComposite)
 				{
 					// All windows have a texture so we can composite
-					CompositeWindows();
+					CompositeWindows(UserData);
 				}
 			}
 		}
 	}
 
-	void FVideoProducerBackBufferComposited::CompositeWindows()
+	void FVideoProducerBackBufferComposited::CompositeWindows(TSharedPtr<FVideoProducerUserData> UserData)
 	{
 		// Process all of the windows we will need to render. This processing step finds the extents of the
 		// composited texture as well as the top-left point
@@ -259,7 +263,8 @@ namespace UE::EditorPixelStreaming2
 			GraphBuilder.Execute();
 		}
 
-		PushFrame(FPixelCaptureInputFrameRHI(OutTexture->GetRHI()));
+		UserData->ProductionEndCycles = FPlatformTime::Cycles64();
+		PushFrame(FPixelCaptureInputFrameRHI(OutTexture->GetRHI(), UserData));
 
 		// Update any subscribed streamers to let them know our composited frame size and position. This way it can correctly scale input from the browser
 		*SharedFrameRect.Get() = FIntRect(TopLeft, BottomRight);
diff --git a/Engine/Plugins/Media/PixelStreaming2/Source/PixelStreaming2Editor/Private/VideoProducerViewport.cpp b/Engine/Plugins/Media/PixelStreaming2/Source/PixelStreaming2Editor/Private/VideoProducerViewport.cpp
--- a/Engine/Plugins/Media/PixelStreaming2/Source/PixelStreaming2Editor/Private/VideoProducerViewport.cpp
+++ b/Engine/Plugins/Media/PixelStreaming2/Source/PixelStreaming2Editor/Private/VideoProducerViewport.cpp
@@ -154,9 +154,14 @@ namespace UE::EditorPixelStreaming2
 			return;
 		}
 
+		TSharedPtr<FVideoProducerUserData> UserData = MakeShared<FVideoProducerUserData>();
+		UserData->ProducerName = TEXT("FVideoProducerViewport");
+		UserData->ProductionBeginCycles = FPlatformTime::Cycles64();
+
 		ENQUEUE_RENDER_COMMAND(StreamViewportTextureCommand)
-		([this, FrameBuffer](FRHICommandList& RHICmdList) {
-			PushFrame(FPixelCaptureInputFrameRHI(FrameBuffer));
+		([this, FrameBuffer, UserData](FRHICommandList& RHICmdList) {
+			UserData->ProductionEndCycles = FPlatformTime::Cycles64();
+			PushFrame(FPixelCaptureInputFrameRHI(FrameBuffer, UserData));
 		});
 	}
 
diff --git a/Engine/Plugins/Media/PixelStreaming2/Source/PixelStreaming2RTC/Private/Stats.cpp b/Engine/Plugins/Media/PixelStreaming2/Source/PixelStreaming2RTC/Private/Stats.cpp
--- a/Engine/Plugins/Media/PixelStreaming2/Source/PixelStreaming2RTC/Private/Stats.cpp
+++ b/Engine/Plugins/Media/PixelStreaming2/Source/PixelStreaming2RTC/Private/Stats.cpp
@@ -74,7 +74,7 @@ namespace UE::PixelStreaming2
 				GraphValue(StatName, Stat.GetValue<double>(), 60, 0, Stat.GetValue<double>() * 10.0f, 0);
 			}
 
-			if(Stat.IsNumeric())
+			if (Stat.IsNumeric())
 			{
 				// If a stat has an alias, use that as the storage key, otherwise use its underlying name
 				FireStatChanged(PlayerId, StatName, Stat.GetValue<double>());
@@ -89,7 +89,7 @@ namespace UE::PixelStreaming2
 		if (FPeerStats* SinglePeerStats = PeerStats.Find(PlayerId))
 		{
 			TMap<FName, FStatGroup>& StatGroups = SinglePeerStats->GetStatGroups();
-			TArray<FName>				   StatCategories;
+			TArray<FName>			 StatCategories;
 			StatGroups.GetKeys(StatCategories);
 
 			// Stat groups contain a name as well as additional info like track index and ssrc
@@ -569,26 +569,35 @@ namespace UE::PixelStreaming2
 		}
 	}
 
-	double FStats::AddTimeStat(uint64 Millis, const FString& Label)
+	double FStats::AddTimeStat(double Millis, const FString& Label)
 	{
-		const double	DeltaMs = Millis;
-		const FStat TimeData({ .Name = FName(*Label) }, DeltaMs, 2, true);
+		const double DeltaMs = Millis;
+		const FStat	 TimeData({ .Name = FName(*Label) }, DeltaMs, 2, true);
 		StoreApplicationStat(TimeData);
 		return DeltaMs;
 	}
 
 	double FStats::AddTimeDeltaStat(uint64 Millis1, uint64 Millis2, const FString& Label)
 	{
-		const uint64	MaxMillis = FGenericPlatformMath::Max(Millis1, Millis2);
-		const uint64	MinMillis = FGenericPlatformMath::Min(Millis1, Millis2);
-		const double	DeltaMs = (MaxMillis - MinMillis) * ((Millis1 > Millis2) ? 1.0 : -1.0);
-		const FStat TimeData({ .Name = FName(*Label) }, DeltaMs, 2, true);
+		const uint64 MaxMillis = FGenericPlatformMath::Max(Millis1, Millis2);
+		const uint64 MinMillis = FGenericPlatformMath::Min(Millis1, Millis2);
+		const double DeltaMs = (MaxMillis - MinMillis) * ((Millis1 > Millis2) ? 1.0 : -1.0);
+		const FStat	 TimeData({ .Name = FName(*Label) }, DeltaMs, 2, true);
 		StoreApplicationStat(TimeData);
 		return DeltaMs;
 	}
 
 	void FStats::AddFrameTimingStats(const FPixelCaptureFrameMetadata& FrameMetadata)
 	{
+		const int Samples = 100;
+
+		TSharedPtr<FVideoProducerUserData> UserData = StaticCastSharedPtr<FVideoProducerUserData>(FrameMetadata.UserData);
+		if (UserData)
+		{
+			const double TimeProduce = AddTimeStat(FPlatformTime::ToMilliseconds64(UserData->ProductionEndCycles - UserData->ProductionBeginCycles), FString::Printf(TEXT("%s Frame Production Time"), *(UserData->ProducerName)));
+			GraphValue(*FString::Printf(TEXT("%d Produce Time"), FrameMetadata.Layer), StaticCast<float>(TimeProduce), Samples, 0.0f, 30.0f);
+		}
+
 		const double TimeCapture = AddTimeStat(FrameMetadata.CaptureTime, FString::Printf(TEXT("%s Layer %d Frame Capture Time"), *FrameMetadata.ProcessName, FrameMetadata.Layer));
 		const double TimeCPU = AddTimeStat(FrameMetadata.CaptureProcessCPUTime, FString::Printf(TEXT("%s Layer %d Frame Capture CPU Time"), *FrameMetadata.ProcessName, FrameMetadata.Layer));
 		const double TimeGPUDelay = AddTimeStat(FrameMetadata.CaptureProcessGPUDelay, FString::Printf(TEXT("%s Layer %d Frame Capture GPU Delay Time"), *FrameMetadata.ProcessName, FrameMetadata.Layer));
@@ -597,7 +606,6 @@ namespace UE::PixelStreaming2
 		const FStat UseData({ .Name = FName(*FString::Printf(TEXT("%s Layer %d Frame Uses"), *FrameMetadata.ProcessName, FrameMetadata.Layer)) }, static_cast<double>(FrameMetadata.UseCount));
 		StoreApplicationStat(UseData);
 
-		const int Samples = 100;
 		GraphValue(*FString::Printf(TEXT("%d Capture Time"), FrameMetadata.Layer), StaticCast<float>(TimeCapture), Samples, 0.0f, 30.0f);
 		GraphValue(*FString::Printf(TEXT("%d CPU Time"), FrameMetadata.Layer), StaticCast<float>(TimeCPU), Samples, 0.0f, 30.0f);
 		GraphValue(*FString::Printf(TEXT("%d GPU Delay Time"), FrameMetadata.Layer), StaticCast<float>(TimeGPUDelay), Samples, 0.0f, 30.0f);
diff --git a/Engine/Plugins/Media/PixelStreaming2/Source/PixelStreaming2RTC/Private/Stats.h b/Engine/Plugins/Media/PixelStreaming2/Source/PixelStreaming2RTC/Private/Stats.h
--- a/Engine/Plugins/Media/PixelStreaming2/Source/PixelStreaming2RTC/Private/Stats.h
+++ b/Engine/Plugins/Media/PixelStreaming2/Source/PixelStreaming2RTC/Private/Stats.h
@@ -103,7 +103,7 @@ namespace UE::PixelStreaming2
 
 		FORCEINLINE TStatId GetStatId() const { RETURN_QUICK_DECLARE_CYCLE_STAT(PixelStreaming2Stats, STATGROUP_Tickables); }
 
-		double AddTimeStat(uint64 Millis, const FString& Label);
+		double AddTimeStat(double Millis, const FString& Label);
 		double AddTimeDeltaStat(uint64 Millis1, uint64 Millis2, const FString& Label);
 		void   AddFrameTimingStats(const FPixelCaptureFrameMetadata& FrameMetadata);
 		void   AddCanvasTile(FName Name, const FCanvasTileItem& Tile);
