[PS2] Fix signalling server unable to connect when port is already in use.

This fixes the unit tests when PS2 would try to connect to a used port.

diff --git a/Engine/Plugins/Media/PixelStreaming2/Source/PixelStreaming2RTC/Private/Tests/TestUtils.cpp b/Engine/Plugins/Media/PixelStreaming2/Source/PixelStreaming2RTC/Private/Tests/TestUtils.cpp
--- a/Engine/Plugins/Media/PixelStreaming2/Source/PixelStreaming2RTC/Private/Tests/TestUtils.cpp
+++ b/Engine/Plugins/Media/PixelStreaming2/Source/PixelStreaming2RTC/Private/Tests/TestUtils.cpp
@@ -11,15 +11,19 @@
 #include "PixelStreaming2PluginSettings.h"
 #include "UtilsAsync.h"
 #include "VideoProducer.h"
+#include "SocketUtils.h"
+
+#include "SocketSubsystem.h"
 
 #if WITH_DEV_AUTOMATION_TESTS
 
 namespace UE::PixelStreaming2
-{
+{	
 	int32 TestUtils::NextStreamerPort()
 	{
 		// Start of IANA un-registerable ports (49152 - 65535)
 		static int NextStreamerPort = 49152;
+		NextStreamerPort = GetNextAvailablePort(NextStreamerPort);
 		return NextStreamerPort++;
 	}
 
@@ -27,6 +31,7 @@ namespace UE::PixelStreaming2
 	{
 		// Half of IANA un-registerable ports (49152 - 65535)
 		static int NextPlayerPort = 57344;
+		NextPlayerPort = GetNextAvailablePort(NextPlayerPort);
 		return NextPlayerPort++;
 	}
 
diff --git a/Engine/Plugins/Media/PixelStreaming2/Source/PixelStreaming2Servers/Internal/SocketUtils.h b/Engine/Plugins/Media/PixelStreaming2/Source/PixelStreaming2Servers/Internal/SocketUtils.h
new file mode 100644
--- /dev/null
+++ b/Engine/Plugins/Media/PixelStreaming2/Source/PixelStreaming2Servers/Internal/SocketUtils.h
@@ -0,0 +1,10 @@
+// Copyright Epic Games, Inc. All Rights Reserved.
+
+#pragma once
+
+#include "Misc/Optional.h"
+
+namespace UE::PixelStreaming2
+{
+	PIXELSTREAMING2SERVERS_API extern int GetNextAvailablePort(TOptional<int> StartingPort = FNullOpt(0));
+} // namespace UE::PixelStreaming2
diff --git a/Engine/Plugins/Media/PixelStreaming2/Source/PixelStreaming2Servers/PixelStreaming2Servers.build.cs b/Engine/Plugins/Media/PixelStreaming2/Source/PixelStreaming2Servers/PixelStreaming2Servers.build.cs
--- a/Engine/Plugins/Media/PixelStreaming2/Source/PixelStreaming2Servers/PixelStreaming2Servers.build.cs
+++ b/Engine/Plugins/Media/PixelStreaming2/Source/PixelStreaming2Servers/PixelStreaming2Servers.build.cs
@@ -24,6 +24,7 @@ namespace UnrealBuildTool.Rules
 				"Json",
 				"WebSocketNetworking",
 				"WebSockets",
+				"Sockets"
 			});
 
 			PrivateDependencyModuleNames.AddRange(new string[] {});
diff --git a/Engine/Plugins/Media/PixelStreaming2/Source/PixelStreaming2Servers/Private/Logging.h b/Engine/Plugins/Media/PixelStreaming2/Source/PixelStreaming2Servers/Private/Logging.h
--- a/Engine/Plugins/Media/PixelStreaming2/Source/PixelStreaming2Servers/Private/Logging.h
+++ b/Engine/Plugins/Media/PixelStreaming2/Source/PixelStreaming2Servers/Private/Logging.h
@@ -4,5 +4,6 @@
 
 #include "HAL/IConsoleManager.h"
 #include "Logging/LogMacros.h"
+#include "Logging/StructuredLog.h"
 
 DECLARE_LOG_CATEGORY_EXTERN(LogPixelStreaming2Servers, Log, All);
\ No newline at end of file
diff --git a/Engine/Plugins/Media/PixelStreaming2/Source/PixelStreaming2Servers/Private/PixelStreaming2ServersModule.cpp b/Engine/Plugins/Media/PixelStreaming2/Source/PixelStreaming2Servers/Private/PixelStreaming2ServersModule.cpp
--- a/Engine/Plugins/Media/PixelStreaming2/Source/PixelStreaming2Servers/Private/PixelStreaming2ServersModule.cpp
+++ b/Engine/Plugins/Media/PixelStreaming2/Source/PixelStreaming2Servers/Private/PixelStreaming2ServersModule.cpp
@@ -12,10 +12,4 @@ void FPixelStreaming2ServersModule::StartupModule()
 	// Any startup logic require goes here.
 }
 
-int FPixelStreaming2ServersModule::NextPort()
-{
-	// TODO: RTCP-7026 Checking for in-use ports.
-	return (4000 + NextGeneratedPort.Increment()) % 65535;
-}
-
 IMPLEMENT_MODULE(FPixelStreaming2ServersModule, PixelStreaming2Servers)
diff --git a/Engine/Plugins/Media/PixelStreaming2/Source/PixelStreaming2Servers/Private/PixelStreaming2ServersModule.h b/Engine/Plugins/Media/PixelStreaming2/Source/PixelStreaming2Servers/Private/PixelStreaming2ServersModule.h
--- a/Engine/Plugins/Media/PixelStreaming2/Source/PixelStreaming2Servers/Private/PixelStreaming2ServersModule.h
+++ b/Engine/Plugins/Media/PixelStreaming2/Source/PixelStreaming2Servers/Private/PixelStreaming2ServersModule.h
@@ -11,12 +11,8 @@ class FPixelStreaming2ServersModule : public IModuleInterface
 public:
 	virtual ~FPixelStreaming2ServersModule() = default;
 	virtual void StartupModule() override;
-	int NextPort();
 
 public:
 	static FPixelStreaming2ServersModule& Get();
 
-private:
-	FThreadSafeCounter NextGeneratedPort = 0;
-
 };
\ No newline at end of file
diff --git a/Engine/Plugins/Media/PixelStreaming2/Source/PixelStreaming2Servers/Private/ServerUtils.cpp b/Engine/Plugins/Media/PixelStreaming2/Source/PixelStreaming2Servers/Private/ServerUtils.cpp
--- a/Engine/Plugins/Media/PixelStreaming2/Source/PixelStreaming2Servers/Private/ServerUtils.cpp
+++ b/Engine/Plugins/Media/PixelStreaming2/Source/PixelStreaming2Servers/Private/ServerUtils.cpp
@@ -11,6 +11,7 @@
 #include "Serialization/JsonReader.h"
 #include "Policies/CondensedJsonPrintPolicy.h"
 #include "Misc/Paths.h"
+#include "SocketUtils.h"
 #include "GenericPlatform/GenericPlatformFile.h"
 
 namespace UE::PixelStreaming2Servers::Utils
@@ -204,14 +205,33 @@ namespace UE::PixelStreaming2Servers::Utils
 		return LaunchChildProcess(OutScriptPath, Args, FString(TEXT("Download ps servers")), true /*bRunAsScript*/);
 	}
 
-	void PopulateCirrusEndPoints(FLaunchArgs& InLaunchArgs, TMap<EEndpoint, FURL>& OutEndPoints)
+	bool PopulateCirrusEndPoints(FLaunchArgs& InLaunchArgs, TMap<EEndpoint, FURL>& OutEndPoints)
 	{
 		FPixelStreaming2ServersModule& Module = FPixelStreaming2ServersModule::Get();
 
 		// Query for ports, or set them if they don't exist
-		FString StreamerPort = Utils::QueryOrSetProcessArgs(InLaunchArgs, TEXT("--StreamerPort="), FString::FromInt(Module.NextPort()));
-		FString SFUPort = Utils::QueryOrSetProcessArgs(InLaunchArgs, TEXT("--SFUPort="), FString::FromInt(Module.NextPort()));
-		FString HttpPort = Utils::QueryOrSetProcessArgs(InLaunchArgs, TEXT("--HttpPort="), FString::FromInt(Module.NextPort()));
+		int FallbackStreamerPort = UE::PixelStreaming2::GetNextAvailablePort();
+		int FallbackSFUPort = UE::PixelStreaming2::GetNextAvailablePort();
+		int FallbackHttpPort = UE::PixelStreaming2::GetNextAvailablePort();
+		int FallbackHttpsPort = UE::PixelStreaming2::GetNextAvailablePort();
+		if (FallbackStreamerPort == -1)
+		{
+			UE_LOGFMT(LogPixelStreaming2Servers, Warning, "Failed to find an available port for streamer connections");
+			return false;
+		}
+		if (FallbackSFUPort == -1)
+		{
+			UE_LOGFMT(LogPixelStreaming2Servers, Warning, "Failed to find an available port for SFU connections");
+			return false;
+		}
+		if (FallbackHttpPort == -1)
+		{
+			UE_LOGFMT(LogPixelStreaming2Servers, Warning, "Failed to find an available port for http connections");
+			return false;
+		}
+		FString StreamerPort = Utils::QueryOrSetProcessArgs(InLaunchArgs, TEXT("--StreamerPort="), FString::FromInt(FallbackStreamerPort));
+		FString SFUPort = Utils::QueryOrSetProcessArgs(InLaunchArgs, TEXT("--SFUPort="), FString::FromInt(FallbackSFUPort));
+		FString HttpPort = Utils::QueryOrSetProcessArgs(InLaunchArgs, TEXT("--HttpPort="), FString::FromInt(FallbackHttpPort));
 		FString ServeHttps = Utils::QueryOrSetProcessArgs(InLaunchArgs, TEXT("--ServeHttps="), TEXT("false"));
 		bool	bServeHttps = ServeHttps == TEXT("true");
 
@@ -249,6 +269,8 @@ namespace UE::PixelStreaming2Servers::Utils
 		OutEndPoints.Add(EEndpoint::Signalling_Players, PlayersUrl);
 		OutEndPoints.Add(EEndpoint::Signalling_SFU, SFUUrl);
 		OutEndPoints.Add(EEndpoint::Signalling_Webserver, WebServerUrl);
+
+		return true;
 	}
 
 	FString ToString(FURL Url)
diff --git a/Engine/Plugins/Media/PixelStreaming2/Source/PixelStreaming2Servers/Private/ServerUtils.h b/Engine/Plugins/Media/PixelStreaming2/Source/PixelStreaming2Servers/Private/ServerUtils.h
--- a/Engine/Plugins/Media/PixelStreaming2/Source/PixelStreaming2Servers/Private/ServerUtils.h
+++ b/Engine/Plugins/Media/PixelStreaming2/Source/PixelStreaming2Servers/Private/ServerUtils.h
@@ -67,8 +67,9 @@ namespace UE::PixelStreaming2Servers
 		 * Populates the cirrus endpoints by extracting endpoint information from the `InLaunchArgs.ProcessArgs`, and if not found, using defaults.
 		 * @param InLaunchArgs The launch args containing the process args that we will be inspecting.
 		 * @param OutEndPoints The end points for the server if we launch it with these launch args.
+		 * @param return True if the end point was added.
 		 **/
-		void PopulateCirrusEndPoints(FLaunchArgs& InLaunchArgs, TMap<EEndpoint, FURL>& OutEndPoints);
+		bool PopulateCirrusEndPoints(FLaunchArgs& InLaunchArgs, TMap<EEndpoint, FURL>& OutEndPoints);
 
 		/**
 		 * Convert a `FURL` into a `FString`. `FURL::ToString()` already exists; however, it appends a problematic trailing slash to the output string.
diff --git a/Engine/Plugins/Media/PixelStreaming2/Source/PixelStreaming2Servers/Private/SocketUtils.cpp b/Engine/Plugins/Media/PixelStreaming2/Source/PixelStreaming2Servers/Private/SocketUtils.cpp
new file mode 100644
--- /dev/null
+++ b/Engine/Plugins/Media/PixelStreaming2/Source/PixelStreaming2Servers/Private/SocketUtils.cpp
@@ -0,0 +1,33 @@
+// Copyright Epic Games, Inc. All Rights Reserved.
+
+#include "SocketUtils.h"
+
+#include "Logging.h"
+#include "SocketSubsystem.h"
+
+namespace UE::PixelStreaming2
+{
+	static FThreadSafeCounter NextGeneratedPort = 0;
+
+	int	GetNextAvailablePort(TOptional<int> StartingPort)
+	{
+		int32					  CandidateNextPort = StartingPort.Get((4000 + NextGeneratedPort.Increment()) % 65535);
+		int32					  NumRemainingPorts = 65535 - CandidateNextPort;
+		ISocketSubsystem*		  SocketSubsystem = ISocketSubsystem::Get(PLATFORM_SOCKETSUBSYSTEM);
+		TSharedRef<FInternetAddr> LocalHostAddr = SocketSubsystem->CreateInternetAddr(FNetworkProtocolTypes::IPv4);
+		LocalHostAddr->SetPort(CandidateNextPort);
+
+		// The following logic is done inside a scope to ensure that the FUniqueSocket is destroyed and unbound before we return the port number
+		{
+			FUniqueSocket Socket = SocketSubsystem->CreateUniqueSocket(NAME_Stream, TEXT("DummySocket"), FNetworkProtocolTypes::IPv4);
+			int32		  BoundPort = SocketSubsystem->BindNextPort(Socket.Get(), LocalHostAddr.Get(), NumRemainingPorts, 1);
+			if (BoundPort == 0)
+			{
+				UE_LOGFMT(LogPixelStreaming2Servers, Warning, "Failed to find an available port!");
+				return -1;
+			}
+			CandidateNextPort = BoundPort;
+		}
+		return CandidateNextPort;
+	}
+} // namespace UE::PixelStreaming2
