[PS2] PS2 Delegates would crash when closing down on linux.

Fix was to move the delegate static functions to cpp file.

Problem is the Get function would create multiple instances of the singleton.

diff --git a/Engine/Plugins/Media/PixelStreaming2/Source/PixelStreaming2/Private/PixelStreaming2Delegates.cpp b/Engine/Plugins/Media/PixelStreaming2/Source/PixelStreaming2/Private/PixelStreaming2Delegates.cpp
new file mode 100644
--- /dev/null
+++ b/Engine/Plugins/Media/PixelStreaming2/Source/PixelStreaming2/Private/PixelStreaming2Delegates.cpp
@@ -0,0 +1,20 @@
+// Copyright Epic Games, Inc. All Rights Reserved.
+
+#include "PixelStreaming2Delegates.h"
+
+UPixelStreaming2Delegates* UPixelStreaming2Delegates::Singleton = nullptr;
+
+UPixelStreaming2Delegates* UPixelStreaming2Delegates::Get()
+{
+	if (Singleton == nullptr && !IsEngineExitRequested())
+	{
+		Singleton = NewObject<UPixelStreaming2Delegates>();
+		Singleton->AddToRoot();
+	}
+	return Singleton;
+}
+
+UPixelStreaming2Delegates::~UPixelStreaming2Delegates()
+{
+	Singleton = nullptr;
+}
diff --git a/Engine/Plugins/Media/PixelStreaming2/Source/PixelStreaming2/Public/PixelStreaming2Delegates.h b/Engine/Plugins/Media/PixelStreaming2/Source/PixelStreaming2/Public/PixelStreaming2Delegates.h
--- a/Engine/Plugins/Media/PixelStreaming2/Source/PixelStreaming2/Public/PixelStreaming2Delegates.h
+++ b/Engine/Plugins/Media/PixelStreaming2/Source/PixelStreaming2/Public/PixelStreaming2Delegates.h
@@ -163,22 +163,11 @@ public:
 	 * @param Get the UPixelStreaming2Delegates singleton.
 	 * @return UPixelStreaming2Delegates pointer.
 	 */
-	static UPixelStreaming2Delegates* Get()
-	{
-		if (Singleton == nullptr && !IsEngineExitRequested())
-		{
-			Singleton = NewObject<UPixelStreaming2Delegates>();
-			Singleton->AddToRoot();
-		}
-		return Singleton;
-	}
-
-	virtual ~UPixelStreaming2Delegates()
-	{
-		Singleton = nullptr;
-	}
+	static UPixelStreaming2Delegates* Get();
+
+	virtual ~UPixelStreaming2Delegates();
 
 private:
 	// The singleton object.
-	static inline UPixelStreaming2Delegates* Singleton = nullptr;
+	static UPixelStreaming2Delegates* Singleton;
 };
