[PS2] Fix: Streamer disconnect/reconnect flow causing lifetime issues

diff --git a/Engine/Plugins/Media/PixelStreaming2/Source/PixelStreaming2RTC/Private/EpicRtcStreamer.cpp b/Engine/Plugins/Media/PixelStreaming2/Source/PixelStreaming2RTC/Private/EpicRtcStreamer.cpp
--- a/Engine/Plugins/Media/PixelStreaming2/Source/PixelStreaming2RTC/Private/EpicRtcStreamer.cpp
+++ b/Engine/Plugins/Media/PixelStreaming2/Source/PixelStreaming2RTC/Private/EpicRtcStreamer.cpp
@@ -195,6 +195,12 @@ namespace UE::PixelStreaming2
 			Delegates->OnUseMediaCaptureChanged.AddSP(Streamer.ToSharedRef(), &FEpicRtcStreamer::OnUseMediaCaptureChanged);
 		}
 
+		if (UPixelStreaming2Delegates* Delegates = UPixelStreaming2Delegates::Get())
+		{
+			Delegates->OnStatChangedNative.AddSP(Streamer.ToSharedRef(), &FEpicRtcStreamer::ConsumeStats);
+			Delegates->OnAllConnectionsClosedNative.AddSP(Streamer.ToSharedRef(), &FEpicRtcStreamer::TriggerMouseLeave);
+		}
+
 		FPixelStreaming2RTCModule::GetModule()->GetStatsCollector()->OnStatsReady.AddSP(Streamer.ToSharedRef(), &FEpicRtcStreamer::OnStatsReady);
 	}
 
@@ -242,25 +248,25 @@ namespace UE::PixelStreaming2
 
 	void FEpicRtcStreamer::StartStreaming()
 	{
-		if (CurrentSignallingServerURL.IsEmpty())
+		if (StreamState != EStreamState::Disconnected)
 		{
-			UE_LOG(LogPixelStreaming2RTC, Log, TEXT("Attempted to start streamer (%s) but no signalling server URL has been set. Use Streamer->SetConnectionURL(URL) or -PixelStreamingConnectionURL="), *StreamerId);
+			UE_LOGFMT(LogPixelStreaming2RTC, Log, "Streamer is already streaming. Ignoring subsequent call to StartStreaming!");
 			return;
 		}
+		StreamState = EStreamState::Connecting;
 
-		StopStreaming();
-		ReconnectTimer->Stop();
-
-		if (UPixelStreaming2Delegates* Delegates = UPixelStreaming2Delegates::Get())
+		if (CurrentSignallingServerURL.IsEmpty())
 		{
-			ConsumeStatsHandle = Delegates->OnStatChangedNative.AddSP(AsShared(), &FEpicRtcStreamer::ConsumeStats);
-			AllConnectionsClosedHandle = Delegates->OnAllConnectionsClosedNative.AddSP(AsShared(), &FEpicRtcStreamer::TriggerMouseLeave);
+			UE_LOGFMT(LogPixelStreaming2RTC, Warning, "Attempted to start streamer ({0}) but no signalling server URL has been set. Use Streamer->SetConnectionURL(URL) or -PixelStreamingConnectionURL=", StreamerId);
+			return;
 		}
 
+		ReconnectTimer->Stop();
+
 		VideoCapturer->ResetFrameCapturer();
 
 		// Broadcast the preconnection event just before we do `TryConnect`
-		StreamingPreConnectionEvent.Broadcast(this);
+		OnPreConnection().Broadcast(this);
 
 		VideoSourceGroup->Start();
 
@@ -273,87 +279,40 @@ namespace UE::PixelStreaming2
 			._observer = EpicRtcManager->SessionObserver
 		};
 
-		EpicRtcErrorCode Result = EpicRtcManager->EpicRtcConference->CreateSession(SessionConfig, EpicRtcManager->EpicRtcSession.GetInitReference());
-		if (Result != EpicRtcErrorCode::Ok)
+		if (const EpicRtcErrorCode Result = EpicRtcManager->EpicRtcConference->CreateSession(SessionConfig, EpicRtcManager->EpicRtcSession.GetInitReference()); Result != EpicRtcErrorCode::Ok)
 		{
-			UE_LOG(LogPixelStreaming2RTC, Error, TEXT("Failed to create EpicRtc session. CreateSession returned %s"), *ToString(Result));
-			StopStreaming();
+			UE_LOGFMT(LogPixelStreaming2RTC, Error, "Failed to create EpicRtc session. CreateSession returned {0}", *ToString(Result));
+			StreamState = EStreamState::Disconnected;
 			return;
 		}
 
-		Result = EpicRtcManager->EpicRtcSession->Connect();
-		if (Result != EpicRtcErrorCode::Ok)
+		// TODO (william.belcher): This should move to OnSessionStateUpdate(EpicRtcSessionState::New) once EpicRtc starts broadcasting that state
+		if (const EpicRtcErrorCode Result = EpicRtcManager->EpicRtcSession->Connect(); Result != EpicRtcErrorCode::Ok)
 		{
-			UE_LOG(LogPixelStreaming2RTC, Error, TEXT("Failed to connect EpicRtcSession. Connect returned %s"), *ToString(Result));
-			StopStreaming();
+			UE_LOGFMT(LogPixelStreaming2RTC, Error, "Failed to connect EpicRtcSession. Connect returned {0}", ToString(Result));
+			StreamState = EStreamState::Disconnected;
 			return;
 		}
-
-		bStreamingStarted = true;
 	}
 
 	void FEpicRtcStreamer::StopStreaming()
 	{
-		if (UPixelStreaming2Delegates* Delegates = UPixelStreaming2Delegates::Get())
-		{
-			Delegates->OnStatChangedNative.Remove(ConsumeStatsHandle);
-			Delegates->OnAllConnectionsClosedNative.Remove(AllConnectionsClosedHandle);
-		}
-
-		RemoveSession(true);
-
-		VideoSourceGroup->Stop();
-		TriggerMouseLeave(StreamerId);
-
-		if (bStreamingStarted)
-		{
-			OnStreamingStopped().Broadcast(this);
-		}
-
-		DeleteAllPlayerSessions();
-		bStreamingStarted = false;
-	}
-
-	void FEpicRtcStreamer::RemoveSession(bool bDisconnect)
-	{
-		if (!EpicRtcManager->EpicRtcSession)
+		if (StreamState == EStreamState::Disconnected)
 		{
 			return;
 		}
 
-		RemoveRoom();
+		StreamState = EStreamState::Disconnecting;
 
-		if (bDisconnect)
+		if (EpicRtcManager->EpicRtcRoom)
 		{
-			EpicRtcErrorCode Result = EpicRtcManager->EpicRtcSession->Disconnect(ToEpicRtcStringView("Streaming Session Removed"));
-			if (Result == EpicRtcErrorCode::SessionDisconnected)
-			{
-				UE_LOG(LogPixelStreaming2RTC, VeryVerbose, TEXT("FEpicRtcStreamer::StopStreaming - Session disconnected cleanly."));
-			}
-			else if (Result != EpicRtcErrorCode::Ok)
-			{
-				UE_LOG(LogPixelStreaming2RTC, Error, TEXT("Failed to disconnect EpicRtcSession. Disconnect returned %s"), *ToString(Result));
-			}
+			EpicRtcManager->EpicRtcRoom->Leave();
 		}
 
-		FUtf8String Utf8StreamerId(StreamerId);
-		EpicRtcManager->EpicRtcConference->RemoveSession(ToEpicRtcStringView(Utf8StreamerId));
-
-		EpicRtcManager->EpicRtcSession = nullptr;
-	}
-
-	void FEpicRtcStreamer::RemoveRoom()
-	{
-		if (!EpicRtcManager->EpicRtcRoom)
-		{
-			return;
-		}
-
-		FUtf8String Utf8StreamerId(StreamerId);
-		EpicRtcManager->EpicRtcRoom->Leave();
-		EpicRtcManager->EpicRtcSession->RemoveRoom(ToEpicRtcStringView(Utf8StreamerId));
+		VideoSourceGroup->Stop();
+		TriggerMouseLeave(StreamerId);
 
-		EpicRtcManager->EpicRtcRoom = nullptr;
+		DeleteAllPlayerSessions();
 	}
 
 	void FEpicRtcStreamer::OnStatsReady(const FString& PlayerId, const EpicRtcConnectionStats& ConnectionStats)
@@ -1010,16 +969,26 @@ namespace UE::PixelStreaming2
 	{
 		switch (State)
 		{
-			case EpicRtcSessionState::Connected:
+			case EpicRtcSessionState::New: // Indicates newly created session.
+			{
+				UE_LOGFMT(LogPixelStreaming2RTC, VeryVerbose, "FEpicRtcStreamer::OnSessionStateUpdate State=New");
+				break;
+			}
+			case EpicRtcSessionState::Pending: // Indicates connection is in progress.
+			{
+				UE_LOGFMT(LogPixelStreaming2RTC, VeryVerbose, "FEpicRtcStreamer::OnSessionStateUpdate State=Pending");
+				break;
+			}
+			case EpicRtcSessionState::Connected: // Indicates session is connected to signalling server.
 			{
-				bSignallingConnected = true;
+				UE_LOGFMT(LogPixelStreaming2RTC, VeryVerbose, "FEpicRtcStreamer::OnSessionStateUpdate State=Connected");
+
 				if (UPixelStreaming2Delegates* Delegates = UPixelStreaming2Delegates::Get())
 				{
 					Delegates->OnConnectedToSignallingServer.Broadcast(StreamerId);
 					Delegates->OnConnectedToSignallingServerNative.Broadcast(StreamerId);
 				}
 
-				UE_LOG(LogPixelStreaming2RTC, VeryVerbose, TEXT("FEpicRtcStreamer::OnSessionStateUpdate State=Connected"));
 				EpicRtcBitrate Bitrate = GetBitrates();
 
 				EpicRtcPortAllocator PortAllocator = {
@@ -1052,49 +1021,64 @@ namespace UE::PixelStreaming2
 				EpicRtcErrorCode Result = EpicRtcManager->EpicRtcSession->CreateRoom(RoomConfig, EpicRtcManager->EpicRtcRoom.GetInitReference());
 				if (Result != EpicRtcErrorCode::Ok)
 				{
-					UE_LOG(LogPixelStreaming2RTC, Error, TEXT("Failed to create EpicRtc room. CreateRoom returned %s"), *ToString(Result));
+					UE_LOGFMT(LogPixelStreaming2RTC, Error, "Failed to create EpicRtc room. CreateRoom returned {0}", ToString(Result));
 					break;
 				}
 
+				// TODO (william.belcher): This should move to OnRoomStateUpdate(EpicRtcRoomState::New) once EpicRtc starts broadcasting that state
 				EpicRtcManager->EpicRtcRoom->Join();
-
-				// Would be better renamed to OnSessionConnected
-				OnStreamingStarted().Broadcast(this);
-			}
-			break;
-			case EpicRtcSessionState::New:
-				// Do something on `new` Session here
-				UE_LOG(LogPixelStreaming2RTC, VeryVerbose, TEXT("FEpicRtcStreamer::OnSessionStateUpdate State=New"));
 				break;
-			case EpicRtcSessionState::Pending:
-				// Do something on `pending` session here
-				UE_LOG(LogPixelStreaming2RTC, VeryVerbose, TEXT("FEpicRtcStreamer::OnSessionStateUpdate State=Pending"));
-				break;
-			case EpicRtcSessionState::Disconnected:
-				// Do something on `disconnected` session here
-				bSignallingConnected = false;
+			}
+			case EpicRtcSessionState::Disconnected: // Indicates session is disconnected from the signalling server.
+			{
+				UE_LOGFMT(LogPixelStreaming2RTC, VeryVerbose, "FEpicRtcStreamer::OnSessionStateUpdate State=Disconnected");
+
+				// If the session unexpectedly disconnects (ie signalling server goes away), we should try and reconnect
+				if (StreamState == EStreamState::Connected)
+				{
+					// Call stop streaming first. This will remove the room
+					StopStreaming();
+					ReconnectTimer->Start(AsShared());
+				}
+
 				if (UPixelStreaming2Delegates* Delegates = UPixelStreaming2Delegates::Get())
 				{
 					Delegates->OnDisconnectedFromSignallingServer.Broadcast(StreamerId);
 					Delegates->OnDisconnectedFromSignallingServerNative.Broadcast(StreamerId);
 				}
-				UE_LOG(LogPixelStreaming2RTC, VeryVerbose, TEXT("FEpicRtcStreamer::OnSessionStateUpdate State=Disconnected"));
-				RemoveSession(false);
-				StopStreaming();
-				ReconnectTimer->Start(AsShared());
+
+				OnStreamingStopped().Broadcast(this);
+
+				// The session has been disconnected (either through a call to StopStreaming or by an error)
+				// so remove it
+				if (EpicRtcManager->EpicRtcConference)
+				{
+					FUtf8String Utf8StreamerId(StreamerId);
+					EpicRtcManager->EpicRtcConference->RemoveSession(ToEpicRtcStringView(Utf8StreamerId));
+					EpicRtcManager->EpicRtcSession = nullptr;
+				}
+
+				// We are fully disconnect at this point so we update the state so we can StartStreaming again
+				StreamState = EStreamState::Disconnected;
+
 				break;
-			case EpicRtcSessionState::Failed:
-				// Do something on `failed` session here
-				bSignallingConnected = false;
-				UE_LOG(LogPixelStreaming2RTC, VeryVerbose, TEXT("FEpicRtcStreamer::OnSessionStateUpdate State=Failed"));
+			}
+			case EpicRtcSessionState::Failed: // Indicates session failed and is unusable.
+			{
+				UE_LOGFMT(LogPixelStreaming2RTC, VeryVerbose, "FEpicRtcStreamer::OnSessionStateUpdate State=Failed");
 				break;
-			case EpicRtcSessionState::Exiting:
-				// Do something on `exiting` session here
-				UE_LOG(LogPixelStreaming2RTC, VeryVerbose, TEXT("FEpicRtcStreamer::OnSessionStateUpdate State=Exiting"));
+			}
+			case EpicRtcSessionState::Exiting: // Indicates session has terminated without a result as a response to the application exiting.
+			{
+				UE_LOGFMT(LogPixelStreaming2RTC, VeryVerbose, "FEpicRtcStreamer::OnSessionStateUpdate State=Exiting");
 				break;
+			}
 			default:
-				UE_LOG(LogPixelStreaming2RTC, Error, TEXT("FEpicRtcStreamer::OnSessionStateUpdate An unhandled session state was encountered. This switch might be missing a case."));
+			{
+				UE_LOGFMT(LogPixelStreaming2RTC, Error, "FEpicRtcStreamer::OnSessionStateUpdate An unhandled session state was encountered. This switch might be missing a case.");
+				checkNoEntry(); // All cases should be handled
 				break;
+			}
 		}
 	}
 
@@ -1110,7 +1094,72 @@ namespace UE::PixelStreaming2
 
 	void FEpicRtcStreamer::OnRoomStateUpdate(const EpicRtcRoomState State)
 	{
-		UE_LOGFMT(LogPixelStreaming2RTC, VeryVerbose, "FEpicRtcStreamer::OnRoomStateUpdate does nothing");
+		switch (State)
+		{
+			case EpicRtcRoomState::New: // Indicates newly created EpicRtcRoomInterface.
+			{
+				UE_LOGFMT(LogPixelStreaming2RTC, VeryVerbose, "FEpicRtcStreamer::OnRoomStateUpdate State=New");
+
+				break;
+			}
+			case EpicRtcRoomState::Pending: // Indicates join of the local participant is in progress.
+			{
+				UE_LOGFMT(LogPixelStreaming2RTC, VeryVerbose, "FEpicRtcStreamer::OnRoomStateUpdate State=Pending");
+
+				break;
+			}
+			case EpicRtcRoomState::Joined: // Indicates local participant (this streamer) is joined.
+			{
+				UE_LOGFMT(LogPixelStreaming2RTC, VeryVerbose, "FEpicRtcStreamer::OnRoomStateUpdate State=Joined");
+
+				StreamState = EStreamState::Connected;
+				OnStreamingStarted().Broadcast(this);
+				break;
+			}
+			case EpicRtcRoomState::Left: // Indicates local participant (this streamer) has left this EpicRtcRoomInterface. Room is not usable once in this state.
+			{
+				UE_LOGFMT(LogPixelStreaming2RTC, VeryVerbose, "FEpicRtcStreamer::OnRoomStateUpdate State=Left");
+
+				if (!EpicRtcManager->EpicRtcSession)
+				{
+					break;
+				}
+
+				FUtf8String Utf8StreamerId(StreamerId);
+				EpicRtcManager->EpicRtcSession->RemoveRoom(ToEpicRtcStringView(Utf8StreamerId));
+				EpicRtcManager->EpicRtcRoom = nullptr;
+
+				const EpicRtcErrorCode Result = EpicRtcManager->EpicRtcSession->Disconnect(ToEpicRtcStringView("Streaming Session Removed"));
+				if (Result == EpicRtcErrorCode::SessionDisconnected)
+				{
+					UE_LOGFMT(LogPixelStreaming2RTC, VeryVerbose, "Session disconnected cleanly.");
+				}
+				else if (Result != EpicRtcErrorCode::Ok)
+				{
+					UE_LOGFMT(LogPixelStreaming2RTC, Error, "Failed to disconnect EpicRtcSession. Disconnect returned {0}", *ToString(Result));
+				}
+
+				break;
+			}
+			case EpicRtcRoomState::Failed: // Indicates EpicRtcRoomInterface failed and is unusable.
+			{
+				UE_LOGFMT(LogPixelStreaming2RTC, VeryVerbose, "FEpicRtcStreamer::OnRoomStateUpdate State=Failed");
+
+				break;
+			}
+			case EpicRtcRoomState::Exiting: // Indicates EpicRtcRoomInterface has terminated without a result as a response to the application exiting.
+			{
+				UE_LOGFMT(LogPixelStreaming2RTC, VeryVerbose, "FEpicRtcStreamer::OnRoomStateUpdate State=Exiting");
+
+				break;
+			}
+			default:
+			{
+				UE_LOGFMT(LogPixelStreaming2RTC, Error, "FEpicRtcStreamer::OnRoomStateUpdate An unhandled room state was encountered. This switch might be missing a case.");
+				checkNoEntry(); // All cases should be handled
+				break;
+			}
+		}
 	}
 
 	FUtf8String GetAudioStreamID()
diff --git a/Engine/Plugins/Media/PixelStreaming2/Source/PixelStreaming2RTC/Private/EpicRtcStreamer.h b/Engine/Plugins/Media/PixelStreaming2/Source/PixelStreaming2RTC/Private/EpicRtcStreamer.h
--- a/Engine/Plugins/Media/PixelStreaming2/Source/PixelStreaming2RTC/Private/EpicRtcStreamer.h
+++ b/Engine/Plugins/Media/PixelStreaming2/Source/PixelStreaming2RTC/Private/EpicRtcStreamer.h
@@ -41,10 +41,10 @@ namespace UE::PixelStreaming2
 		virtual FString GetConnectionURL() override;
 
 		virtual FString GetId() override { return StreamerId; };
-		virtual bool	IsConnected() override { return bSignallingConnected; }
+		virtual bool	IsConnected() override { return StreamState == EStreamState::Connected; }
 		virtual void	StartStreaming() override;
 		virtual void	StopStreaming() override;
-		virtual bool	IsStreaming() const override { return bStreamingStarted; }
+		virtual bool	IsStreaming() const override { return StreamState == EStreamState::Connected; }
 
 		virtual FPreConnectionEvent&	OnPreConnection() override;
 		virtual FStreamingStartedEvent& OnStreamingStarted() override;
@@ -94,13 +94,12 @@ namespace UE::PixelStreaming2
 		bool FindPlayerByVideoTrack(EpicRtcVideoTrackInterface* VideoTrack, FString& OutPlayerId);
 		bool FindPlayerByDataTrack(EpicRtcDataTrackInterface* DataTrack, FString& OutPlayerId);
 		void TriggerMouseLeave(FString InStreamerId);
-		void RemoveSession(bool bDisconnect);
-		void RemoveRoom();
-		void OnStatsReady(const FString& PlayerId, const EpicRtcConnectionStats& ConnectionStats);
-		void OnFrameCapturerCreated();
-		void OnUIInteraction(FMemoryReader Ar);
-		void OnSendMessage(FString MessageName, FMemoryReader Ar);
-		void OnUseMediaCaptureChanged(IConsoleVariable* CVar);
+
+		void		   OnStatsReady(const FString& PlayerId, const EpicRtcConnectionStats& ConnectionStats);
+		void		   OnFrameCapturerCreated();
+		void		   OnUIInteraction(FMemoryReader Ar);
+		void		   OnSendMessage(FString MessageName, FMemoryReader Ar);
+		void		   OnUseMediaCaptureChanged(IConsoleVariable* CVar);
 		EpicRtcBitrate GetBitrates() const;
 
 	private:
@@ -113,8 +112,15 @@ namespace UE::PixelStreaming2
 
 		FString InputControllingId = INVALID_PLAYER_ID;
 
-		bool bSignallingConnected = false;
-		bool bStreamingStarted = false;
+		enum class EStreamState : uint8
+		{
+			Disconnected = 0,
+			Connecting,
+			Connected,
+			Disconnecting,
+		};
+
+		EStreamState StreamState = EStreamState::Disconnected;
 
 		FPreConnectionEvent	   StreamingPreConnectionEvent;
 		FStreamingStartedEvent StreamingStartedEvent;
@@ -124,9 +130,6 @@ namespace UE::PixelStreaming2
 		TSharedPtr<FVideoSourceGroup>	  VideoSourceGroup;
 		TSharedPtr<FFreezeFrame>		  FreezeFrame;
 
-		FDelegateHandle ConsumeStatsHandle;
-		FDelegateHandle AllConnectionsClosedHandle;
-
 		TMap<FName, FString> ConfigOptions;
 
 		TSharedPtr<FStreamerReconnectTimer> ReconnectTimer;
diff --git a/Engine/Plugins/Media/PixelStreaming2/Source/PixelStreaming2RTC/Private/Tests/CodecUtils.cpp b/Engine/Plugins/Media/PixelStreaming2/Source/PixelStreaming2RTC/Private/Tests/CodecUtils.cpp
--- a/Engine/Plugins/Media/PixelStreaming2/Source/PixelStreaming2RTC/Private/Tests/CodecUtils.cpp
+++ b/Engine/Plugins/Media/PixelStreaming2/Source/PixelStreaming2RTC/Private/Tests/CodecUtils.cpp
@@ -19,24 +19,22 @@ namespace UE::PixelStreaming2
 
 		TSharedPtr<UE::PixelStreaming2Servers::IServer> SignallingServer = CreateSignallingServer(StreamerPort, PlayerPort);
 
-		FString									   StreamerName(FString::Printf(TEXT("MockStreamer%d"), StreamerPort));
+		FString								 StreamerName(FString::Printf(TEXT("MockStreamer%d"), StreamerPort));
 		TSharedPtr<IPixelStreaming2Streamer> Streamer = CreateStreamer(StreamerName, StreamerPort);
-		TSharedPtr<FVideoProducer>				   VideoProducer = FVideoProducer::Create();
+		TSharedPtr<FVideoProducer>			 VideoProducer = FVideoProducer::Create();
 		Streamer->SetVideoProducer(VideoProducer);
 
 		TSharedPtr<FMockPlayer>	   Player = CreatePlayer();
 		TSharedPtr<FMockVideoSink> VideoSink = Player->GetVideoSink();
 
-		Streamer->StartStreaming();
-
-		if (!Streamer->IsStreaming())
-		{
-			UE_LOG(LogPixelStreaming2RTC, Error, TEXT("Streamer is not connected."));
-			return;
-		}
-
-		Player->Connect(PlayerPort);
+		TSharedPtr<bool> bStreamingStarted = MakeShared<bool>(false);
+		Streamer->OnStreamingStarted().AddLambda([bStreamingStarted](IPixelStreaming2Streamer*) {
+			*(bStreamingStarted.Get()) = true;
+		});
 
+		ADD_LATENT_AUTOMATION_COMMAND(FExecuteLambda([Streamer]() { Streamer->StartStreaming(); }))
+		ADD_LATENT_AUTOMATION_COMMAND(FWaitAndCheckStreamerBool(TEXT("Check streaming started"), 5.0, Streamer, bStreamingStarted, true))
+		ADD_LATENT_AUTOMATION_COMMAND(FExecuteLambda([Player, PlayerPort]() { Player->Connect(PlayerPort); }))
 		ADD_LATENT_AUTOMATION_COMMAND(FSubscribePlayerAfterStreamerConnectedOrTimeout(5.0, Streamer, Player, StreamerName))
 
 		// Send 30 frames
@@ -57,9 +55,9 @@ namespace UE::PixelStreaming2
 
 		TSharedPtr<UE::PixelStreaming2Servers::IServer> SignallingServer = CreateSignallingServer(StreamerPort, PlayerPort);
 
-		FString									   StreamerName(FString::Printf(TEXT("MockStreamer"), StreamerPort));
+		FString								 StreamerName(FString::Printf(TEXT("MockStreamer%d"), StreamerPort));
 		TSharedPtr<IPixelStreaming2Streamer> Streamer = CreateStreamer(StreamerName, StreamerPort);
-		TSharedPtr<FVideoProducer>				   VideoProducer = FVideoProducer::Create();
+		TSharedPtr<FVideoProducer>			 VideoProducer = FVideoProducer::Create();
 		Streamer->SetVideoProducer(VideoProducer);
 
 		TSharedPtr<FMockPlayer>	   Player = CreatePlayer();
@@ -67,10 +65,15 @@ namespace UE::PixelStreaming2
 
 		// Note: Important to couple framerate as we are manually passing frames and don't want any cached frames
 		Streamer->SetCoupleFramerate(true);
-		Streamer->StartStreaming();
-
-		Player->Connect(PlayerPort);
-
+		
+		TSharedPtr<bool> bStreamingStarted = MakeShared<bool>(false);
+		Streamer->OnStreamingStarted().AddLambda([bStreamingStarted](IPixelStreaming2Streamer*) {
+			*(bStreamingStarted.Get()) = true;
+		});
+
+		ADD_LATENT_AUTOMATION_COMMAND(FExecuteLambda([Streamer]() { Streamer->StartStreaming(); }))
+		ADD_LATENT_AUTOMATION_COMMAND(FWaitAndCheckStreamerBool(TEXT("Check streaming started"), 5.0, Streamer, bStreamingStarted, true))
+		ADD_LATENT_AUTOMATION_COMMAND(FExecuteLambda([Player, PlayerPort]() { Player->Connect(PlayerPort); }))
 		ADD_LATENT_AUTOMATION_COMMAND(FSubscribePlayerAfterStreamerConnectedOrTimeout(5.0, Streamer, Player, StreamerName))
 
 		for (int Res = 2; Res < 512; Res *= 2)
diff --git a/Engine/Plugins/Media/PixelStreaming2/Source/PixelStreaming2RTC/Private/Tests/DataChannelTest.cpp b/Engine/Plugins/Media/PixelStreaming2/Source/PixelStreaming2RTC/Private/Tests/DataChannelTest.cpp
--- a/Engine/Plugins/Media/PixelStreaming2/Source/PixelStreaming2RTC/Private/Tests/DataChannelTest.cpp
+++ b/Engine/Plugins/Media/PixelStreaming2/Source/PixelStreaming2RTC/Private/Tests/DataChannelTest.cpp
@@ -34,13 +34,10 @@ namespace UE::PixelStreaming2
 		uint8	FromStreamerEchoId = InputHandler->GetFromStreamerProtocol()->Find(EPixelStreaming2FromStreamerMessage::TestEcho)->GetID();
 		FString EchoFromStreamerContent = TEXT("StreamWillEchoThis");
 
-		Streamer->StartStreaming();
-
-		if (!Streamer->IsStreaming())
-		{
-			UE_LOG(LogPixelStreaming2RTC, Error, TEXT("Streamer is not connected."));
-			return false;
-		}
+		TSharedPtr<bool> bStreamingStarted = MakeShared<bool>(false);
+		Streamer->OnStreamingStarted().AddLambda([bStreamingStarted](IPixelStreaming2Streamer*) {
+			*(bStreamingStarted.Get()) = true;
+		});
 
 		// Setup listener on the Player side for "FromStreamer" messages
 		TSharedPtr<bool>					  bGotMessageFromStreamer = MakeShared<bool>(false);
@@ -57,9 +54,11 @@ namespace UE::PixelStreaming2
 			}
 		};
 		Player->OnMessageReceived.AddLambda(Callback);
-		Player->Connect(PlayerPort);
 
 		// The streamer send an "echo" message to the player. The player then sends this message back to the streamer and we check that we receive this echo.
+		ADD_LATENT_AUTOMATION_COMMAND(FExecuteLambda([Streamer]() { Streamer->StartStreaming(); }))
+		ADD_LATENT_AUTOMATION_COMMAND(FWaitAndCheckStreamerBool(TEXT("Check streaming started"), 5.0, Streamer, bStreamingStarted, true))
+		ADD_LATENT_AUTOMATION_COMMAND(FExecuteLambda([Player, PlayerPort]() { Player->Connect(PlayerPort); }))
 		ADD_LATENT_AUTOMATION_COMMAND(FSubscribePlayerAfterStreamerConnectedOrTimeout(5.0, Streamer, Player, StreamerName));
 		ADD_LATENT_AUTOMATION_COMMAND(FWaitForDataChannelOrTimeout(5.0, Player));
 		ADD_LATENT_AUTOMATION_COMMAND(FSendDataChannelMessageToStreamer(Player, EPixelStreaming2ToStreamerMessage::TestEcho, EchoFromStreamerContent));
diff --git a/Engine/Plugins/Media/PixelStreaming2/Source/PixelStreaming2RTC/Private/Tests/EpicRtcStreamerTest.cpp b/Engine/Plugins/Media/PixelStreaming2/Source/PixelStreaming2RTC/Private/Tests/EpicRtcStreamerTest.cpp
new file mode 100644
--- /dev/null
+++ b/Engine/Plugins/Media/PixelStreaming2/Source/PixelStreaming2RTC/Private/Tests/EpicRtcStreamerTest.cpp
@@ -0,0 +1,166 @@
+// Copyright Epic Games, Inc. All Rights Reserved.
+
+#include "EpicRtcUtils.h"
+#include "EpicRtcWebsocket.h"
+#include "Logging.h"
+#include "Misc/AutomationTest.h"
+#include "Tests/AutomationCommon.h"
+#include "TestUtils.h"
+#include "UtilsString.h"
+#include "PixelStreaming2Delegates.h"
+#include "PixelStreaming2PluginSettings.h"
+
+#if WITH_DEV_AUTOMATION_TESTS
+
+namespace UE::PixelStreaming2
+{
+	IMPLEMENT_SIMPLE_AUTOMATION_TEST(FPS2EpicRtcStreamerDelegateTest, "System.Plugins.PixelStreaming2.FPS2EpicRtcStreamerDelegateTest", EAutomationTestFlags::EditorContext | EAutomationTestFlags::ClientContext | EAutomationTestFlags::ProductFilter)
+	bool FPS2EpicRtcStreamerDelegateTest::RunTest(const FString& Parameters)
+	{
+		int32 StreamerPort = TestUtils::NextStreamerPort();
+		int32 PlayerPort = TestUtils::NextPlayerPort();
+
+		TSharedPtr<UE::PixelStreaming2Servers::IServer> SignallingServer = CreateSignallingServer(StreamerPort, PlayerPort);
+
+		FString								 StreamerName(FString::Printf(TEXT("MockStreamer%d"), StreamerPort));
+		TSharedPtr<IPixelStreaming2Streamer> Streamer = CreateStreamer(StreamerName, StreamerPort);
+
+		UPixelStreaming2Delegates* Delegates = UPixelStreaming2Delegates::Get();
+		if (!Delegates)
+		{
+			UE_LOGFMT(LogPixelStreaming2RTC, Error, "Failed to obtain delegates pointer. FPS2EpicRtcStreamerDelegateTest will not continue!");
+			return true;
+		}
+
+		TSharedPtr<bool> bConnectedToSignallingServerNative = MakeShared<bool>(false);
+		Delegates->OnConnectedToSignallingServerNative.AddLambda([StreamerName, bConnectedToSignallingServerNative](FString ConnectedStreamer) {
+			if (ConnectedStreamer == StreamerName)
+			{
+				*(bConnectedToSignallingServerNative.Get()) = true;
+			}
+		});
+
+		TSharedPtr<bool> bDisconnectedFromSignallingServerNative = MakeShared<bool>(false);
+		Delegates->OnDisconnectedFromSignallingServerNative.AddLambda([StreamerName, bDisconnectedFromSignallingServerNative](FString ConnectedStreamer) {
+			if (ConnectedStreamer == StreamerName)
+			{
+				*(bDisconnectedFromSignallingServerNative.Get()) = true;
+			}
+		});
+
+		TSharedPtr<bool> bStreamingStarted = MakeShared<bool>(false);
+		Streamer->OnStreamingStarted().AddLambda([bStreamingStarted](IPixelStreaming2Streamer*) {
+			*(bStreamingStarted.Get()) = true;
+		});
+
+		TSharedPtr<bool> bStreamingStopped = MakeShared<bool>(false);
+		Streamer->OnStreamingStopped().AddLambda([bStreamingStopped](IPixelStreaming2Streamer*) {
+			*(bStreamingStopped.Get()) = true;
+		});
+
+		// Start streaming
+		ADD_LATENT_AUTOMATION_COMMAND(FExecuteLambda([Streamer]() { Streamer->StartStreaming(); }))
+		ADD_LATENT_AUTOMATION_COMMAND(FWaitAndCheckStreamerBool(TEXT("Check streaming started"), 5.0, Streamer, bStreamingStarted, true))
+		ADD_LATENT_AUTOMATION_COMMAND(FWaitAndCheckStreamerBool(TEXT("Check connected to signalling server"), 5.0, Streamer, bConnectedToSignallingServerNative, true))
+		// Wait 1 second to ensure any websocket message have correctly flowed
+		ADD_LATENT_AUTOMATION_COMMAND(FWaitSeconds(1.0))
+		// Stop streaming
+		ADD_LATENT_AUTOMATION_COMMAND(FExecuteLambda([Streamer]() { Streamer->StopStreaming(); }))
+		ADD_LATENT_AUTOMATION_COMMAND(FWaitAndCheckStreamerBool(TEXT("Check streaming stopped"), 5.0, Streamer, bStreamingStopped, true))
+		ADD_LATENT_AUTOMATION_COMMAND(FWaitAndCheckStreamerBool(TEXT("Check disconnected from signalling server"), 5.0, Streamer, bDisconnectedFromSignallingServerNative, true))
+		// Cleanup
+		ADD_LATENT_AUTOMATION_COMMAND(FCleanupAll(SignallingServer, Streamer, nullptr))
+
+		return true;
+	}
+
+	IMPLEMENT_SIMPLE_AUTOMATION_TEST(FPS2EpicRtcStreamerReconnectTest, "System.Plugins.PixelStreaming2.FPS2EpicRtcStreamerReconnectTest", EAutomationTestFlags::EditorContext | EAutomationTestFlags::ClientContext | EAutomationTestFlags::ProductFilter)
+	bool FPS2EpicRtcStreamerReconnectTest::RunTest(const FString& Parameters)
+	{
+		int32 StreamerPort = TestUtils::NextStreamerPort();
+		int32 PlayerPort = TestUtils::NextPlayerPort();
+
+		TSharedPtr<UE::PixelStreaming2Servers::IServer> SignallingServer = CreateSignallingServer(StreamerPort, PlayerPort);
+
+		FString								 StreamerName(FString::Printf(TEXT("MockStreamer%d"), StreamerPort));
+		TSharedPtr<IPixelStreaming2Streamer> Streamer = CreateStreamer(StreamerName, StreamerPort);
+
+		UPixelStreaming2Delegates* Delegates = UPixelStreaming2Delegates::Get();
+		if (!Delegates)
+		{
+			UE_LOGFMT(LogPixelStreaming2RTC, Error, "Failed to obtain delegates pointer. FPS2EpicRtcStreamerReconnectTest will not continue!");
+			return true;
+		}
+
+		TSharedPtr<bool> bConnectedToSignallingServerNative = MakeShared<bool>(false);
+		Delegates->OnConnectedToSignallingServerNative.AddLambda([StreamerName, bConnectedToSignallingServerNative](FString ConnectedStreamer) {
+			if (ConnectedStreamer == StreamerName)
+			{
+				*(bConnectedToSignallingServerNative.Get()) = true;
+			}
+		});
+
+		TSharedPtr<bool> bDisconnectedFromSignallingServerNative = MakeShared<bool>(false);
+		Delegates->OnDisconnectedFromSignallingServerNative.AddLambda([StreamerName, bDisconnectedFromSignallingServerNative](FString ConnectedStreamer) {
+			if (ConnectedStreamer == StreamerName)
+			{
+				*(bDisconnectedFromSignallingServerNative.Get()) = true;
+			}
+		});
+
+		TSharedPtr<bool> bStreamingStarted = MakeShared<bool>(false);
+		Streamer->OnStreamingStarted().AddLambda([bStreamingStarted](IPixelStreaming2Streamer*) {
+			*(bStreamingStarted.Get()) = true;
+		});
+
+		TSharedPtr<bool> bStreamingStopped = MakeShared<bool>(false);
+		Streamer->OnStreamingStopped().AddLambda([bStreamingStopped](IPixelStreaming2Streamer*) {
+			*(bStreamingStopped.Get()) = true;
+		});
+
+		// Start streaming
+		ADD_LATENT_AUTOMATION_COMMAND(FExecuteLambda([Streamer]() { Streamer->StartStreaming(); }))
+		ADD_LATENT_AUTOMATION_COMMAND(FWaitAndCheckStreamerBool(TEXT("Check streaming started"), 5.0, Streamer, bStreamingStarted, true))
+		ADD_LATENT_AUTOMATION_COMMAND(FWaitAndCheckStreamerBool(TEXT("Check connected to signalling server"), 5.0, Streamer, bConnectedToSignallingServerNative, true))
+		// Wait 1 second to ensure any websocket message have correctly flowed
+		ADD_LATENT_AUTOMATION_COMMAND(FWaitSeconds(1.0))
+		// Reset the signalling server, this will trigger the reconnection flow
+		ADD_LATENT_AUTOMATION_COMMAND(FExecuteLambda([SignallingServer]() { SignallingServer->Stop(); }))
+		ADD_LATENT_AUTOMATION_COMMAND(FWaitAndCheckStreamerBool(TEXT("Check streaming stopped from SS going away"), 5.0, Streamer, bStreamingStopped, true))
+		ADD_LATENT_AUTOMATION_COMMAND(FWaitAndCheckStreamerBool(TEXT("Check disconnected from signalling server from SS going away"), 5.0, Streamer, bDisconnectedFromSignallingServerNative, true))
+		// Reset the state variables
+		ADD_LATENT_AUTOMATION_COMMAND(FExecuteLambda([bStreamingStarted, bConnectedToSignallingServerNative, bStreamingStopped, bDisconnectedFromSignallingServerNative]() {
+			*(bStreamingStarted.Get()) = false;
+			*(bConnectedToSignallingServerNative.Get()) = false;
+			*(bStreamingStopped.Get()) = false;
+			*(bDisconnectedFromSignallingServerNative.Get()) = false;
+		}))
+		// Restart the signalling server, the streamer should reconnect
+		ADD_LATENT_AUTOMATION_COMMAND(FExecuteLambda([SignallingServer, StreamerPort, PlayerPort]() {
+			UE::PixelStreaming2Servers::FLaunchArgs LaunchArgs;
+			LaunchArgs.ProcessArgs = FString::Printf(TEXT("--StreamerPort=%d --HttpPort=%d"), StreamerPort, PlayerPort);
+			bool bLaunchedSignallingServer = SignallingServer->Launch(LaunchArgs);
+			if (!bLaunchedSignallingServer)
+			{
+				UE_LOG(LogPixelStreaming2RTC, Error, TEXT("Failed to relaunch signalling server."));
+			}
+		}))
+		// Wait to ensure any websocket message have correctly flowed. Time is 2 times the reconnect interval for safety
+		ADD_LATENT_AUTOMATION_COMMAND(FWaitSeconds(2 * UPixelStreaming2PluginSettings::CVarSignalingReconnectInterval.GetValueOnAnyThread()))
+		ADD_LATENT_AUTOMATION_COMMAND(FWaitAndCheckStreamerBool(TEXT("Check streaming restarted"), 5.0, Streamer, bStreamingStarted, true))
+		ADD_LATENT_AUTOMATION_COMMAND(FWaitAndCheckStreamerBool(TEXT("Check reconnected to signalling server"), 5.0, Streamer, bConnectedToSignallingServerNative, true))
+		// Wait 1 second to ensure any websocket message have correctly flowed
+		ADD_LATENT_AUTOMATION_COMMAND(FWaitSeconds(1.0))
+		// Stop streaming
+		ADD_LATENT_AUTOMATION_COMMAND(FExecuteLambda([Streamer]() { Streamer->StopStreaming(); }))
+		ADD_LATENT_AUTOMATION_COMMAND(FWaitAndCheckStreamerBool(TEXT("Check streaming stopped"), 5.0, Streamer, bStreamingStopped, true))
+		ADD_LATENT_AUTOMATION_COMMAND(FWaitAndCheckStreamerBool(TEXT("Check disconnected from signalling server"), 5.0, Streamer, bDisconnectedFromSignallingServerNative, true))
+
+		// Cleanup
+		ADD_LATENT_AUTOMATION_COMMAND(FCleanupAll(SignallingServer, Streamer, nullptr))
+
+		return true;
+	}
+} // namespace UE::PixelStreaming2
+
+#endif
diff --git a/Engine/Plugins/Media/PixelStreaming2/Source/PixelStreaming2RTC/Private/Tests/MediaDirectionTest.cpp b/Engine/Plugins/Media/PixelStreaming2/Source/PixelStreaming2RTC/Private/Tests/MediaDirectionTest.cpp
--- a/Engine/Plugins/Media/PixelStreaming2/Source/PixelStreaming2RTC/Private/Tests/MediaDirectionTest.cpp
+++ b/Engine/Plugins/Media/PixelStreaming2/Source/PixelStreaming2RTC/Private/Tests/MediaDirectionTest.cpp
@@ -51,13 +51,10 @@ namespace UE::PixelStreaming2
 
 		TSharedPtr<FMockPlayer> Player = CreatePlayer(PlayerConfig);
 
-		Streamer->StartStreaming();
-
-		if (!Streamer->IsStreaming())
-		{
-			UE_LOG(LogPixelStreaming2RTC, Error, TEXT("Streamer is not connected."));
-			return;
-		}
+		TSharedPtr<bool> bStreamingStarted = MakeShared<bool>(false);
+		Streamer->OnStreamingStarted().AddLambda([bStreamingStarted](IPixelStreaming2Streamer*) {
+			*(bStreamingStarted.Get()) = true;
+		});
 
 		TSharedPtr<bool> bLocalTrack = MakeShared<bool>(false);
 		TSharedPtr<bool> bRemoteTrack = MakeShared<bool>(false);
@@ -88,8 +85,9 @@ namespace UE::PixelStreaming2
 			});
 		}
 
-		Player->Connect(PlayerPort);
-
+		ADD_LATENT_AUTOMATION_COMMAND(FExecuteLambda([Streamer]() { Streamer->StartStreaming(); }))
+		ADD_LATENT_AUTOMATION_COMMAND(FWaitAndCheckStreamerBool(TEXT("Check streaming started"), 5.0, Streamer, bStreamingStarted, true))
+		ADD_LATENT_AUTOMATION_COMMAND(FExecuteLambda([Player, PlayerPort]() { Player->Connect(PlayerPort); }))
 		ADD_LATENT_AUTOMATION_COMMAND(FSubscribePlayerAfterStreamerConnectedOrTimeout(5.0, Streamer, Player, StreamerName))
 
 		if (Direction == EMediaDirection::SendOnly)
diff --git a/Engine/Plugins/Media/PixelStreaming2/Source/PixelStreaming2RTC/Private/Tests/MultipleSignallingConnectionsTest.cpp b/Engine/Plugins/Media/PixelStreaming2/Source/PixelStreaming2RTC/Private/Tests/MultipleSignallingConnectionsTest.cpp
--- a/Engine/Plugins/Media/PixelStreaming2/Source/PixelStreaming2RTC/Private/Tests/MultipleSignallingConnectionsTest.cpp
+++ b/Engine/Plugins/Media/PixelStreaming2/Source/PixelStreaming2RTC/Private/Tests/MultipleSignallingConnectionsTest.cpp
@@ -210,14 +210,8 @@ namespace UE::PixelStreaming2
 		Streamer1->OnStreamingStopped().AddLambda([stream1Disconnected](IPixelStreaming2Streamer*) { *stream1Disconnected.Get() = true; });
 		Streamer2->OnStreamingStopped().AddLambda([stream2Disconnected](IPixelStreaming2Streamer*) { *stream2Disconnected.Get() = true; });
 
-		Streamer1->StartStreaming();
-		Streamer2->StartStreaming();
-
-		if (!Streamer1->IsStreaming() || !Streamer2->IsStreaming())
-		{
-			UE_LOG(LogPixelStreaming2RTC, Error, TEXT("Streamer is not connected."));
-			return false;
-		}
+		ADD_LATENT_AUTOMATION_COMMAND(FExecuteLambda([Streamer1]() { Streamer1->StartStreaming(); }))
+		ADD_LATENT_AUTOMATION_COMMAND(FExecuteLambda([Streamer2]() { Streamer2->StartStreaming(); }))
 
 		ADD_LATENT_AUTOMATION_COMMAND(FWaitForServerOrTimeout(SignallingServer));
 		ADD_LATENT_AUTOMATION_COMMAND(FWaitForStreamConnected(5.0f, Streamer1, stream1Connected, stream1Disconnected));
diff --git a/Engine/Plugins/Media/PixelStreaming2/Source/PixelStreaming2RTC/Private/Tests/ProtocolTest.cpp b/Engine/Plugins/Media/PixelStreaming2/Source/PixelStreaming2RTC/Private/Tests/ProtocolTest.cpp
--- a/Engine/Plugins/Media/PixelStreaming2/Source/PixelStreaming2RTC/Private/Tests/ProtocolTest.cpp
+++ b/Engine/Plugins/Media/PixelStreaming2/Source/PixelStreaming2RTC/Private/Tests/ProtocolTest.cpp
@@ -20,7 +20,7 @@ namespace UE::PixelStreaming2
 	{
 		// need to be able to accept codec to handshake otherwise setting local description fails when generating an answer
 		SetCodec(EVideoCodec::VP8);
-		
+
 		int32 StreamerPort = TestUtils::NextStreamerPort();
 		int32 PlayerPort = TestUtils::NextPlayerPort();
 
@@ -74,17 +74,16 @@ namespace UE::PixelStreaming2
 			}
 		};
 
-		Streamer->StartStreaming();
-
-		if (!Streamer->IsStreaming())
-		{
-			UE_LOG(LogPixelStreaming2RTC, Error, TEXT("Streamer is not connected."));
-			return false;
-		}
+		TSharedPtr<bool> bStreamingStarted = MakeShared<bool>(false);
+		Streamer->OnStreamingStarted().AddLambda([bStreamingStarted](IPixelStreaming2Streamer*) {
+			*(bStreamingStarted.Get()) = true;
+		});
 
 		Player->OnMessageReceived.AddLambda(Callback);
-		Player->Connect(PlayerPort);
 
+		ADD_LATENT_AUTOMATION_COMMAND(FExecuteLambda([Streamer]() { Streamer->StartStreaming(); }))
+		ADD_LATENT_AUTOMATION_COMMAND(FWaitAndCheckStreamerBool(TEXT("Check streaming started"), 5.0, Streamer, bStreamingStarted, true))
+		ADD_LATENT_AUTOMATION_COMMAND(FExecuteLambda([Player, PlayerPort]() { Player->Connect(PlayerPort); }))
 		ADD_LATENT_AUTOMATION_COMMAND(FSubscribePlayerAfterStreamerConnectedOrTimeout(5.0, Streamer, Player, StreamerName));
 		ADD_LATENT_AUTOMATION_COMMAND(FWaitForDataChannelOrTimeout(5.0, Player));
 		ADD_LATENT_AUTOMATION_COMMAND(FSendCustomMessageToStreamer(Player, CustomMessageName, 1337));
@@ -98,18 +97,18 @@ namespace UE::PixelStreaming2
 	{
 		// need to be able to accept codec to handshake otherwise setting local description fails when generating an answer
 		SetCodec(EVideoCodec::VP8);
-		
+
 		int32 StreamerPort = TestUtils::NextStreamerPort();
 		int32 PlayerPort = TestUtils::NextPlayerPort();
 
 		TSharedPtr<UE::PixelStreaming2Servers::IServer> SignallingServer = CreateSignallingServer(StreamerPort, PlayerPort);
 
-		FString										   StreamerName(FString::Printf(TEXT("MockStreamer"), StreamerPort));
-		TSharedPtr<IPixelStreaming2Streamer>	   Streamer = CreateStreamer(StreamerName, StreamerPort);
+		FString									 StreamerName(FString::Printf(TEXT("MockStreamer%d"), StreamerPort));
+		TSharedPtr<IPixelStreaming2Streamer>	 Streamer = CreateStreamer(StreamerName, StreamerPort);
 		TSharedPtr<IPixelStreaming2InputHandler> InputHandler = Streamer->GetInputHandler().Pin();
 
 		// Define our message and add it to the protocol
-		const FString								   CustomMessageName = FString(TEXT("CustomMessage"));
+		const FString							 CustomMessageName = FString(TEXT("CustomMessage"));
 		TSharedPtr<IPixelStreaming2InputMessage> Message = InputHandler->GetToStreamerProtocol()->Add(CustomMessageName, { EPixelStreaming2MessageTypes::Uint16 });
 
 		// Define a handler function
@@ -129,16 +128,14 @@ namespace UE::PixelStreaming2
 		Player->GetToStreamerProtocol()->Add(CustomMessageName, { EPixelStreaming2MessageTypes::Uint16 });
 		TSharedPtr<FMockVideoSink> VideoSink = Player->GetVideoSink();
 
-		Streamer->StartStreaming();
-
-		if (!Streamer->IsStreaming())
-		{
-			UE_LOG(LogPixelStreaming2RTC, Error, TEXT("Streamer is not connected."));
-			return false;
-		}
-
-		Player->Connect(PlayerPort);
+		TSharedPtr<bool> bStreamingStarted = MakeShared<bool>(false);
+		Streamer->OnStreamingStarted().AddLambda([bStreamingStarted](IPixelStreaming2Streamer*) {
+			*(bStreamingStarted.Get()) = true;
+		});
 
+		ADD_LATENT_AUTOMATION_COMMAND(FExecuteLambda([Streamer]() { Streamer->StartStreaming(); }))
+		ADD_LATENT_AUTOMATION_COMMAND(FWaitAndCheckStreamerBool(TEXT("Check streaming started"), 5.0, Streamer, bStreamingStarted, true))
+		ADD_LATENT_AUTOMATION_COMMAND(FExecuteLambda([Player, PlayerPort]() { Player->Connect(PlayerPort); }))
 		ADD_LATENT_AUTOMATION_COMMAND(FSubscribePlayerAfterStreamerConnectedOrTimeout(5.0, Streamer, Player, StreamerName));
 		ADD_LATENT_AUTOMATION_COMMAND(FWaitForDataChannelOrTimeout(5.0, Player));
 		ADD_LATENT_AUTOMATION_COMMAND(FSendCustomMessageToStreamer(Player, CustomMessageName, 1337));
diff --git a/Engine/Plugins/Media/PixelStreaming2/Source/PixelStreaming2RTC/Private/Tests/TestUtils.cpp b/Engine/Plugins/Media/PixelStreaming2/Source/PixelStreaming2RTC/Private/Tests/TestUtils.cpp
--- a/Engine/Plugins/Media/PixelStreaming2/Source/PixelStreaming2RTC/Private/Tests/TestUtils.cpp
+++ b/Engine/Plugins/Media/PixelStreaming2/Source/PixelStreaming2RTC/Private/Tests/TestUtils.cpp
@@ -248,6 +248,20 @@ namespace UE::PixelStreaming2
 		return false; // We keep updating this until the timeout is complete and then we check that there's no tracks added
 	}
 
+	bool FWaitAndCheckStreamerBool::Update()
+	{
+		double DeltaTime = FPlatformTime::Seconds() - StartTime;
+		if (DeltaTime > WaitSeconds)
+		{
+			if ((*bCheck.Get()) != bExpectedValue)
+			{
+				UE_LOGFMT(LogPixelStreaming2RTC, Error, "{0} failed. Expected [{1}] but got [{2}]", CheckName, bExpectedValue, *bCheck.Get());
+			}
+			return true;
+		}
+		return false; // We keep updating this until the timeout is complete and then we check that there's no tracks added
+	}
+
 	bool FWaitForStreamerDataChannelMessageOrTimeout::Update()
 	{
 		double DeltaTime = FPlatformTime::Seconds() - StartTime;
@@ -300,6 +314,12 @@ namespace UE::PixelStreaming2
 		return true;
 	}
 
+	bool FExecuteLambda::Update()
+	{
+		Func();
+		return true;
+	}
+
 	/* ---------- Utility functions ----------- */
 
 	void SetCodec(EVideoCodec Codec)
diff --git a/Engine/Plugins/Media/PixelStreaming2/Source/PixelStreaming2RTC/Private/Tests/TestUtils.h b/Engine/Plugins/Media/PixelStreaming2/Source/PixelStreaming2RTC/Private/Tests/TestUtils.h
--- a/Engine/Plugins/Media/PixelStreaming2/Source/PixelStreaming2RTC/Private/Tests/TestUtils.h
+++ b/Engine/Plugins/Media/PixelStreaming2/Source/PixelStreaming2RTC/Private/Tests/TestUtils.h
@@ -86,6 +86,8 @@ namespace UE::PixelStreaming2
 	DEFINE_LATENT_AUTOMATION_COMMAND_THREE_PARAMETER(FSendCustomMessageToStreamer, TSharedPtr<FMockPlayer>, Player, FString, MessageType, uint16, Body);
 	DEFINE_LATENT_AUTOMATION_COMMAND_THREE_PARAMETER(FWaitForPlayerTrackOrTimeout, double, TimeoutSeconds, TSharedPtr<FMockPlayer>, Player, TSharedPtr<bool>, bComplete);
 	DEFINE_LATENT_AUTOMATION_COMMAND_FOUR_PARAMETER(FWaitAndCheckBool, double, WaitSeconds, TSharedPtr<FMockPlayer>, Player, TSharedPtr<bool>, bCheck, bool, bExpectedValue);
+	DEFINE_LATENT_AUTOMATION_COMMAND_FIVE_PARAMETER(FWaitAndCheckStreamerBool, const FString, CheckName, double, WaitSeconds, TSharedPtr<IPixelStreaming2Streamer>, Streamer, TSharedPtr<bool>, bCheck, bool, bExpectedValue);
+	DEFINE_LATENT_AUTOMATION_COMMAND_ONE_PARAMETER(FExecuteLambda, const TFunction<void()>, Func);
 
 	TSharedPtr<IPixelStreaming2Streamer>			CreateStreamer(const FString& StreamerName, int StreamerPort);
 	TSharedPtr<FMockPlayer>							CreatePlayer(FMockPlayerConfig Config = { .AudioDirection = EMediaDirection::RecvOnly, .VideoDirection = EMediaDirection::RecvOnly });
