[PS2] Fix: PIE streamer not disabling the CPU throttling setting

diff --git a/Engine/Plugins/Media/PixelStreaming2/Source/PixelStreaming2Editor/Private/PixelStreaming2EditorModule.cpp b/Engine/Plugins/Media/PixelStreaming2/Source/PixelStreaming2Editor/Private/PixelStreaming2EditorModule.cpp
--- a/Engine/Plugins/Media/PixelStreaming2/Source/PixelStreaming2Editor/Private/PixelStreaming2EditorModule.cpp
+++ b/Engine/Plugins/Media/PixelStreaming2/Source/PixelStreaming2Editor/Private/PixelStreaming2EditorModule.cpp
@@ -147,7 +147,7 @@ namespace UE::EditorPixelStreaming2
 		 * Called before the engine exits. Separate from OnPreExit as OnEnginePreExit occurs before shutting down any core modules.
 		 */
 		FCoreDelegates::OnEnginePreExit.AddLambda([this]() {
-			RestoreCPUThrottlingSetting();
+			RestoreCPUThrottlingSetting(true);
 		});
 
 		// Give the editor streamer the default url if the user hasn't specified one when launching the editor
@@ -462,31 +462,39 @@ namespace UE::EditorPixelStreaming2
 		}
 	}
 
-	void FPixelStreaming2EditorModule::RestoreCPUThrottlingSetting()
+	void FPixelStreaming2EditorModule::RestoreCPUThrottlingSetting(bool bForce)
 	{
-		UEditorPerformanceSettings* Settings = GetMutableDefault<UEditorPerformanceSettings>();
-		Settings->bThrottleCPUWhenNotForeground = bOldCPUThrottlingSetting;
-		Settings->PostEditChange();
+		// Test the set count, it can be zero because the restore callback fires on destruct
+		// If bForce is set, this doesn't correspond to a particular call, but should only occur if a disable has occured
+		if (CpuThrottlingSetCount && (bForce || (--CpuThrottlingSetCount == 0)))
+		{
+			UEditorPerformanceSettings* Settings = GetMutableDefault<UEditorPerformanceSettings>();
+			Settings->bThrottleCPUWhenNotForeground = bOldCPUThrottlingSetting;
+			Settings->PostEditChange();
+		}
 	}
 
 	void FPixelStreaming2EditorModule::DisableCPUThrottlingSetting()
 	{
-		// Update editor settings so that editor won't slow down if not in focus
-		UEditorPerformanceSettings* Settings = GetMutableDefault<UEditorPerformanceSettings>();
-
-		// Store whatever value the user had in here so we can restore it when we are done streaming.
-		bOldCPUThrottlingSetting = Settings->bThrottleCPUWhenNotForeground;
-
-		if (Settings->bThrottleCPUWhenNotForeground)
+		if (++CpuThrottlingSetCount == 1)
 		{
-			Settings->bThrottleCPUWhenNotForeground = false;
-			Settings->PostEditChange();
+			// Update editor settings so that editor won't slow down if not in focus
+			UEditorPerformanceSettings* Settings = GetMutableDefault<UEditorPerformanceSettings>();
+
+			// Store whatever value the user had in here so we can restore it when we are done streaming.
+			bOldCPUThrottlingSetting = Settings->bThrottleCPUWhenNotForeground;
 
-			// Let the user know we are forcing this editor setting (so they know why their setting is not working potentially)
-			// We can also use the new keyword here because the notification manager will call delete once the notification has been displayed
-			FNotificationInfo* Info = new FNotificationInfo(LOCTEXT("PixelStreaming2EditorModule_CPUThrottlingNotification", "Pixel Streaming: Disabling setting \"Use less CPU in background\" for streaming performance."));
-			Info->ExpireDuration = 5.0f;
-			FSlateNotificationManager::Get().QueueNotification(Info);
+			if (Settings->bThrottleCPUWhenNotForeground)
+			{
+				Settings->bThrottleCPUWhenNotForeground = false;
+				Settings->PostEditChange();
+
+				// Let the user know we are forcing this editor setting (so they know why their setting is not working potentially)
+				// We can also use the new keyword here because the notification manager will call delete once the notification has been displayed
+				FNotificationInfo* Info = new FNotificationInfo(LOCTEXT("PixelStreaming2EditorModule_CPUThrottlingNotification", "Pixel Streaming: Disabling setting \"Use less CPU in background\" for streaming performance."));
+				Info->ExpireDuration = 5.0f;
+				FSlateNotificationManager::Get().QueueNotification(Info);
+			}
 		}
 	}
 
@@ -502,6 +510,17 @@ namespace UE::EditorPixelStreaming2
 			PIEStreamer->SetConnectionURL(FString::Printf(TEXT("%s:%d"), *SignallingDomain, StreamerPort));
 		}
 
+		// Bind to start/stop streaming so we disable/restore relevant editor settings
+		PIEStreamer->OnStreamingStarted().AddLambda([this](IPixelStreaming2Streamer* Streamer) {
+			DisableCPUThrottlingSetting();
+		});
+		PIEStreamer->OnStreamingStopped().AddLambda([this](IPixelStreaming2Streamer* Streamer) {
+			if (!IsEngineExitRequested())
+			{
+				RestoreCPUThrottlingSetting();
+			}
+		});
+
 		PIEStreamer->SetVideoProducer(FVideoProducerPIEViewport::Create());
 
 		TSharedPtr<IPixelStreaming2InputHandler> InputHandler = PIEStreamer->GetInputHandler().Pin();
diff --git a/Engine/Plugins/Media/PixelStreaming2/Source/PixelStreaming2Editor/Private/PixelStreaming2EditorModule.h b/Engine/Plugins/Media/PixelStreaming2/Source/PixelStreaming2Editor/Private/PixelStreaming2EditorModule.h
--- a/Engine/Plugins/Media/PixelStreaming2/Source/PixelStreaming2Editor/Private/PixelStreaming2EditorModule.h
+++ b/Engine/Plugins/Media/PixelStreaming2/Source/PixelStreaming2Editor/Private/PixelStreaming2EditorModule.h
@@ -7,6 +7,8 @@
 #include "PixelStreaming2Servers.h"
 #include "Widgets/SWindow.h"
 
+#include <atomic>
+
 namespace UE::EditorPixelStreaming2
 {
 	class FPixelStreaming2Toolbar;
@@ -43,7 +45,7 @@ namespace UE::EditorPixelStreaming2
 		bool ParseResolution(const TCHAR* InResolution, uint32& OutX, uint32& OutY);
 		void MaybeResizeEditor(TSharedPtr<SWindow> RootWindow);
 		void DisableCPUThrottlingSetting();
-		void RestoreCPUThrottlingSetting();
+		void RestoreCPUThrottlingSetting(bool bForce = false);
 
 #if WITH_EDITOR
 		void OnBeginPIE(bool bIsSimulating);
@@ -79,6 +81,9 @@ namespace UE::EditorPixelStreaming2
 		TSharedPtr<IPixelStreaming2Streamer> PIEStreamer;
 #endif
 
+		// The original throttling setting prior to first streamer disabling it
 		bool bOldCPUThrottlingSetting = false;
+		// Number of streamers that have requested the throttling be disabled
+		std::atomic<int> CpuThrottlingSetCount = 0;
 	};
 } // namespace UE::EditorPixelStreaming2
