[PS2] Run PS2 blueprint delegates in the game thread.

Fix OnAudioTrackRemoved and OnVideoTrackRemoved not being called by checking Stopped state.

diff --git a/Engine/Plugins/Media/PixelStreaming2/Source/PixelStreaming2/Public/PixelStreaming2Delegates.h b/Engine/Plugins/Media/PixelStreaming2/Source/PixelStreaming2/Public/PixelStreaming2Delegates.h
--- a/Engine/Plugins/Media/PixelStreaming2/Source/PixelStreaming2/Public/PixelStreaming2Delegates.h
+++ b/Engine/Plugins/Media/PixelStreaming2/Source/PixelStreaming2/Public/PixelStreaming2Delegates.h
@@ -26,7 +26,7 @@ public:
 	FConnectedToSignallingServer OnConnectedToSignallingServer;
 	
 	/** Delegate type for when a connection to the signalling server was made. */
-	DECLARE_MULTICAST_DELEGATE_OneParam(FConnectedToSignallingServerNative, FString /* StreamerId */);
+	DECLARE_TS_MULTICAST_DELEGATE_OneParam(FConnectedToSignallingServerNative, FString /* StreamerId */);
 
 	/** Invoked when a connection to the signalling server was made. */
 	FConnectedToSignallingServerNative OnConnectedToSignallingServerNative;
@@ -39,7 +39,7 @@ public:
 	FDisconnectedFromSignallingServer OnDisconnectedFromSignallingServer;
 	
 	/** Delegate type for when a connection to the signalling server was lost. */
-	DECLARE_MULTICAST_DELEGATE_OneParam(FDisconnectedFromSignallingServerNative, FString /* StreamerId */);
+	DECLARE_TS_MULTICAST_DELEGATE_OneParam(FDisconnectedFromSignallingServerNative, FString /* StreamerId */);
 
 	/** Invoked when a connection to the signalling server was lost. */
 	FDisconnectedFromSignallingServerNative OnDisconnectedFromSignallingServerNative;
@@ -141,7 +141,7 @@ public:
 	FStatsChanged OnStatChanged;
 
 	/** Delegate type for when a pixel streaming stat has changed. */
-	DECLARE_MULTICAST_DELEGATE_ThreeParams(FStatChangedNative, FString /* Peer id */, FName /* Stat name */, float);
+	DECLARE_TS_MULTICAST_DELEGATE_ThreeParams(FStatChangedNative, FString /* Peer id */, FName /* Stat name */, float);
 
 	/** Invoked when a pixel streaming stat has changed. */
 	FStatChangedNative OnStatChangedNative;
@@ -154,7 +154,7 @@ public:
 	FFallbackToSoftwareEncoding OnFallbackToSoftwareEncoding;
 
 	/** Delegate type for when the GPU ran out of available hardware encoders and fell back to software encoders. */
-	DECLARE_MULTICAST_DELEGATE(FFallbackToSoftwareEncodingNative);
+	DECLARE_TS_MULTICAST_DELEGATE(FFallbackToSoftwareEncodingNative);
 
 	/** Invoked when the GPU ran out of available hardware encoders and fell back to software encoders. */
 	FFallbackToSoftwareEncodingNative OnFallbackToSoftwareEncodingNative;
diff --git a/Engine/Plugins/Media/PixelStreaming2/Source/PixelStreaming2RTC/Private/EpicRtcStreamer.cpp b/Engine/Plugins/Media/PixelStreaming2/Source/PixelStreaming2RTC/Private/EpicRtcStreamer.cpp
--- a/Engine/Plugins/Media/PixelStreaming2/Source/PixelStreaming2RTC/Private/EpicRtcStreamer.cpp
+++ b/Engine/Plugins/Media/PixelStreaming2/Source/PixelStreaming2RTC/Private/EpicRtcStreamer.cpp
@@ -26,6 +26,7 @@
 #include "EpicRtcDataTrack.h"
 #include "IPixelStreaming2InputModule.h"
 #include "CoreGlobals.h"
+#include "UtilsAsync.h"
 #include "UtilsString.h"
 #include "UtilsCodecs.h"
 #include "VideoProducerBackBuffer.h"
@@ -65,6 +66,18 @@ namespace UE::PixelStreaming2
 	FEpicRtcStreamer::~FEpicRtcStreamer()
 	{
 		StopStreaming();
+
+		// Pin ref counted objects and clear the session if it is still alive.
+		// Session may still be alive if the manager was destroyed quickly after stopping the stream
+		// and the EpicRtcSessionState::Disconnected event did not have enough time to tick.
+		const TRefCountPtr<EpicRtcConferenceInterface> Conference = EpicRtcManager->EpicRtcConference;
+		const TRefCountPtr<EpicRtcSessionInterface> Session = EpicRtcManager->EpicRtcSession;
+		if (Conference.IsValid() && Session.IsValid())
+		{
+			const FUtf8String Utf8StreamerId(StreamerId);
+			Conference->RemoveSession(ToEpicRtcStringView(Utf8StreamerId));
+			EpicRtcManager->EpicRtcSession = nullptr;
+		}
 	}
 
 	void FEpicRtcStreamer::Initialize()
@@ -297,16 +310,16 @@ namespace UE::PixelStreaming2
 
 	void FEpicRtcStreamer::StopStreaming()
 	{
-		if (StreamState == EStreamState::Disconnected)
+		if (StreamState == EStreamState::Disconnected || StreamState == EStreamState::Disconnecting)
 		{
 			return;
 		}
 
 		StreamState = EStreamState::Disconnecting;
 
-		if (EpicRtcManager->EpicRtcRoom)
+		if (const TRefCountPtr<EpicRtcRoomInterface> Room = EpicRtcManager->EpicRtcRoom)
 		{
-			EpicRtcManager->EpicRtcRoom->Leave();
+			Room->Leave();
 		}
 
 		VideoSourceGroup->Stop();
@@ -678,16 +691,18 @@ namespace UE::PixelStreaming2
 
 		Players->Remove(PlayerId);
 
-		if (UPixelStreaming2Delegates* Delegates = UPixelStreaming2Delegates::Get())
-		{
-			Delegates->OnClosedConnection.Broadcast(StreamerId, PlayerId);
-			Delegates->OnClosedConnectionNative.Broadcast(StreamerId, PlayerId);
-			if (Players->IsEmpty())
+		DoOnGameThread([StreamerId = this->StreamerId, PlayerId, bIsEmpty = Players->IsEmpty()]() {
+			if (UPixelStreaming2Delegates* Delegates = UPixelStreaming2Delegates::Get())
 			{
-				Delegates->OnAllConnectionsClosed.Broadcast(StreamerId);
-				Delegates->OnAllConnectionsClosedNative.Broadcast(StreamerId);
+				Delegates->OnClosedConnection.Broadcast(StreamerId, PlayerId);
+				Delegates->OnClosedConnectionNative.Broadcast(StreamerId, PlayerId);
+				if (bIsEmpty)
+				{
+					Delegates->OnAllConnectionsClosed.Broadcast(StreamerId);
+					Delegates->OnAllConnectionsClosedNative.Broadcast(StreamerId);
+				}
 			}
-		}
+		});
 
 		if (FStats* PSStats = FStats::Get())
 		{
@@ -705,27 +720,32 @@ namespace UE::PixelStreaming2
 		VideoSourceGroup->RemoveAllVideoSources();
 		Players->Empty();
 		InputControllingId = INVALID_PLAYER_ID;
-		if (UPixelStreaming2Delegates* Delegates = UPixelStreaming2Delegates::Get())
-		{
-			Delegates->OnAllConnectionsClosed.Broadcast(StreamerId);
-			Delegates->OnAllConnectionsClosedNative.Broadcast(StreamerId);
-		}
+
+		DoOnGameThread([StreamerId = this->StreamerId]() {
+			if (UPixelStreaming2Delegates* Delegates = UPixelStreaming2Delegates::Get())
+			{
+				Delegates->OnAllConnectionsClosed.Broadcast(StreamerId);
+				Delegates->OnAllConnectionsClosedNative.Broadcast(StreamerId);
+			}
+		});
 	}
 
 	void FEpicRtcStreamer::OnDataChannelOpen(FString PlayerId)
 	{
 		// Only time we automatically make a new peer the input controlling host is if they are the first peer (and not the SFU).
-		bool HostControlsInput = GetEnumFromCVar<EInputControllerMode>(UPixelStreaming2PluginSettings::CVarInputController) == EInputControllerMode::Host;
+		const bool HostControlsInput = GetEnumFromCVar<EInputControllerMode>(UPixelStreaming2PluginSettings::CVarInputController) == EInputControllerMode::Host;
 		if (HostControlsInput && !IsSFU(PlayerId) && InputControllingId == INVALID_PLAYER_ID)
 		{
 			InputControllingId = PlayerId;
 		}
 
-		if (UPixelStreaming2Delegates* Delegates = UPixelStreaming2Delegates::Get())
-		{
-			Delegates->OnDataTrackOpen.Broadcast(StreamerId, PlayerId);
-			Delegates->OnDataTrackOpenNative.Broadcast(StreamerId, PlayerId);
-		}
+		DoOnGameThread([StreamerId = this->StreamerId, PlayerId]() {
+			if (UPixelStreaming2Delegates* Delegates = UPixelStreaming2Delegates::Get())
+			{
+				Delegates->OnDataTrackOpen.Broadcast(StreamerId, PlayerId);
+				Delegates->OnDataTrackOpenNative.Broadcast(StreamerId, PlayerId);
+			}
+		});
 
 		// When data channel is open
 		SendProtocol(PlayerId);
@@ -760,11 +780,13 @@ namespace UE::PixelStreaming2
 				});
 			}
 
-			if (UPixelStreaming2Delegates* Delegates = UPixelStreaming2Delegates::Get())
-			{
-				Delegates->OnDataTrackClosed.Broadcast(StreamerId, PlayerId);
-				Delegates->OnDataTrackClosedNative.Broadcast(StreamerId, PlayerId);
-			}
+			DoOnGameThread([StreamerId = this->StreamerId, PlayerId]() {
+				if (UPixelStreaming2Delegates* Delegates = UPixelStreaming2Delegates::Get())
+				{
+					Delegates->OnDataTrackClosed.Broadcast(StreamerId, PlayerId);
+					Delegates->OnDataTrackClosedNative.Broadcast(StreamerId, PlayerId);
+				}
+			});
 		}
 	}
 
@@ -858,7 +880,7 @@ namespace UE::PixelStreaming2
 
 		double ReceiptTimeMs = FPlatformTime::ToMilliseconds64(FPlatformTime::Cycles64());
 
-		AsyncTask(ENamedThreads::GameThread, [this, PlayerId, ReceiptTimeMs]() {
+		DoOnGameThread([this, PlayerId, ReceiptTimeMs]() {
 			FString ReportToTransmitJSON;
 
 			if (!UPixelStreaming2PluginSettings::CVarWebRTCDisableStats.GetValueOnAnyThread())
@@ -958,7 +980,7 @@ namespace UE::PixelStreaming2
 			}
 			else
 			{
-				AsyncTask(ENamedThreads::GameThread, [MouseLeaveFunction]() {
+				DoOnGameThread([MouseLeaveFunction]() {
 					MouseLeaveFunction();
 				});
 			}
@@ -982,14 +1004,15 @@ namespace UE::PixelStreaming2
 			case EpicRtcSessionState::Connected: // Indicates session is connected to signalling server.
 			{
 				UE_LOGFMT(LogPixelStreaming2RTC, VeryVerbose, "FEpicRtcStreamer::OnSessionStateUpdate State=Connected");
+				DoOnGameThread([StreamerId = this->StreamerId]() {
+					if (UPixelStreaming2Delegates* Delegates = UPixelStreaming2Delegates::Get())
+					{
+						Delegates->OnConnectedToSignallingServer.Broadcast(StreamerId);
+						Delegates->OnConnectedToSignallingServerNative.Broadcast(StreamerId);
+					}
+				});
 
-				if (UPixelStreaming2Delegates* Delegates = UPixelStreaming2Delegates::Get())
-				{
-					Delegates->OnConnectedToSignallingServer.Broadcast(StreamerId);
-					Delegates->OnConnectedToSignallingServerNative.Broadcast(StreamerId);
-				}
-
-				EpicRtcBitrate Bitrate = GetBitrates();
+				const EpicRtcBitrate Bitrate = GetBitrates();
 
 				EpicRtcPortAllocator PortAllocator = {
 					._minPort = UPixelStreaming2PluginSettings::CVarWebRTCMinPort.GetValueOnAnyThread(),
@@ -1040,12 +1063,13 @@ namespace UE::PixelStreaming2
 					StopStreaming();
 					ReconnectTimer->Start(AsShared());
 				}
-
-				if (UPixelStreaming2Delegates* Delegates = UPixelStreaming2Delegates::Get())
-				{
-					Delegates->OnDisconnectedFromSignallingServer.Broadcast(StreamerId);
-					Delegates->OnDisconnectedFromSignallingServerNative.Broadcast(StreamerId);
-				}
+				DoOnGameThread([StreamerId = this->StreamerId]() {
+					if (UPixelStreaming2Delegates* Delegates = UPixelStreaming2Delegates::Get())
+					{
+						Delegates->OnDisconnectedFromSignallingServer.Broadcast(StreamerId);
+						Delegates->OnDisconnectedFromSignallingServerNative.Broadcast(StreamerId);
+					}
+				});
 
 				OnStreamingStopped().Broadcast(this);
 
@@ -1176,7 +1200,7 @@ namespace UE::PixelStreaming2
 
 	void FEpicRtcStreamer::OnRoomJoinedUpdate(EpicRtcParticipantInterface* Participant)
 	{
-		FString ParticipantId = ToString(Participant->GetId());
+		const FString ParticipantId = ToString(Participant->GetId());
 		UE_LOG(LogPixelStreaming2RTC, Log, TEXT("Player (%s) joined"), *ParticipantId);
 
 		if (ParticipantId == StreamerId)
@@ -1184,11 +1208,13 @@ namespace UE::PixelStreaming2
 			return;
 		}
 
-		if (UPixelStreaming2Delegates* Delegates = UPixelStreaming2Delegates::Get())
-		{
-			Delegates->OnNewConnection.Broadcast(StreamerId, ParticipantId);
-			Delegates->OnNewConnectionNative.Broadcast(StreamerId, ParticipantId);
-		}
+		DoOnGameThread([StreamerId = this->StreamerId, ParticipantId]() {
+			if (UPixelStreaming2Delegates* Delegates = UPixelStreaming2Delegates::Get())
+			{
+				Delegates->OnNewConnection.Broadcast(StreamerId, ParticipantId);
+				Delegates->OnNewConnectionNative.Broadcast(StreamerId, ParticipantId);
+			}
+		});
 
 		FPlayerContext& PlayerContext = Players->FindOrAdd(ParticipantId);
 		PlayerContext.ParticipantInterface = Participant;
@@ -1527,6 +1553,12 @@ namespace UE::PixelStreaming2
 	}
 
 	void FEpicRtcStreamer::OnAudioTrackRemoved(EpicRtcAudioTrackInterface* AudioTrack)
+	{
+		// TODO (Eden.Harris) OnAudioTrackRemoved is not hooked up in EpicRtc. 
+		UE_LOGFMT(LogPixelStreaming2RTC, Warning, "FEpicRtcStreamer::OnAudioTrackRemoved was called. Use OnAudioTrackState with State [EpicRtcTrackState::Stopped]");
+	}
+
+	void FEpicRtcStreamer::RemoveAudioTrack(EpicRtcAudioTrackInterface* AudioTrack)
 	{
 		FString PlayerId;
 		bool	bFoundPlayer = FindPlayerByAudioTrack(AudioTrack, PlayerId);
@@ -1557,7 +1589,19 @@ namespace UE::PixelStreaming2
 
 	void FEpicRtcStreamer::OnAudioTrackState(EpicRtcAudioTrackInterface* AudioTrack, const EpicRtcTrackState State)
 	{
-		UE_LOGFMT(LogPixelStreaming2RTC, VeryVerbose, "FEpicRtcStreamer::OnAudioTrackState does nothing");
+		FString PlayerId;
+		FindPlayerByAudioTrack(AudioTrack, PlayerId);
+		const FString AudioTrackId = ToString(AudioTrack->GetId());
+
+		if (State == EpicRtcTrackState::Active)
+		{
+			UE_LOGFMT(LogPixelStreaming2RTC, Log, "FEpicRtcStreamer::OnAudioTrackState(AudioTrack=[{0}], Player=[{1}], State=Active)", *AudioTrackId, *PlayerId);
+		}
+		else if (State == EpicRtcTrackState::Stopped)
+		{
+			UE_LOGFMT(LogPixelStreaming2RTC, Log, "FEpicRtcStreamer::OnAudioTrackState(AudioTrack=[{0}], Player=[{1}], State=Stopped)", *AudioTrackId, *PlayerId);
+			RemoveAudioTrack(AudioTrack);
+		}
 	}
 
 	void FEpicRtcStreamer::OnVideoTrackMuted(EpicRtcVideoTrackInterface* VideoTrack, EpicRtcBool bIsMuted)
@@ -1604,6 +1648,12 @@ namespace UE::PixelStreaming2
 	}
 
 	void FEpicRtcStreamer::OnVideoTrackRemoved(EpicRtcVideoTrackInterface* VideoTrack)
+	{
+		// TODO (Eden.Harris) OnVideoTrackRemoved is not hooked up in EpicRtc. 
+		UE_LOGFMT(LogPixelStreaming2RTC, Warning, "FEpicRtcStreamer::OnVideoTrackRemoved was called. Use OnAudioTrackState with State [EpicRtcTrackState::Stopped]");
+	}
+
+	void FEpicRtcStreamer::RemoveVideoTrack(EpicRtcVideoTrackInterface* VideoTrack)
 	{
 		FString PlayerId;
 		bool	bFoundPlayer = FindPlayerByVideoTrack(VideoTrack, PlayerId);
@@ -1637,18 +1687,19 @@ namespace UE::PixelStreaming2
 	void FEpicRtcStreamer::OnVideoTrackState(EpicRtcVideoTrackInterface* VideoTrack, const EpicRtcTrackState State)
 	{
 		FString PlayerId;
-		bool	bFoundPlayer = FindPlayerByVideoTrack(VideoTrack, PlayerId);
-		FString VideoTrackId = ToString(VideoTrack->GetId());
+		FindPlayerByVideoTrack(VideoTrack, PlayerId);
+		const FString VideoTrackId = ToString(VideoTrack->GetId());
 
 		// Note: It is acceptable to not have a found a player for track state changes, as these can trigger before we have added a participant
 
 		if (State == EpicRtcTrackState::Active)
 		{
-			UE_LOG(LogPixelStreaming2RTC, Log, TEXT("FEpicRtcStreamer::OnVideoTrackState(VideoTrack=[%s], Player=[%s], State=Active)"), *VideoTrackId, *PlayerId);
+			UE_LOGFMT(LogPixelStreaming2RTC, Log, "FEpicRtcStreamer::OnVideoTrackState(VideoTrack=[{0}], Player=[{1}], State=Active)", *VideoTrackId, *PlayerId);
 		}
 		else if (State == EpicRtcTrackState::Stopped)
 		{
-			UE_LOG(LogPixelStreaming2RTC, Log, TEXT("FEpicRtcStreamer::OnVideoTrackState(VideoTrack=[%s], Player=[%s], State=Stopped)"), *VideoTrackId, *PlayerId);
+			UE_LOGFMT(LogPixelStreaming2RTC, Log, "FEpicRtcStreamer::OnVideoTrackState(VideoTrack=[{0}], Player=[{1}], State=Stopped)", *VideoTrackId, *PlayerId);
+			RemoveVideoTrack(VideoTrack);
 		}
 	}
 
diff --git a/Engine/Plugins/Media/PixelStreaming2/Source/PixelStreaming2RTC/Private/EpicRtcStreamer.h b/Engine/Plugins/Media/PixelStreaming2/Source/PixelStreaming2RTC/Private/EpicRtcStreamer.h
--- a/Engine/Plugins/Media/PixelStreaming2/Source/PixelStreaming2RTC/Private/EpicRtcStreamer.h
+++ b/Engine/Plugins/Media/PixelStreaming2/Source/PixelStreaming2RTC/Private/EpicRtcStreamer.h
@@ -94,6 +94,8 @@ namespace UE::PixelStreaming2
 		bool FindPlayerByVideoTrack(EpicRtcVideoTrackInterface* VideoTrack, FString& OutPlayerId);
 		bool FindPlayerByDataTrack(EpicRtcDataTrackInterface* DataTrack, FString& OutPlayerId);
 		void TriggerMouseLeave(FString InStreamerId);
+		void RemoveAudioTrack(EpicRtcAudioTrackInterface* AudioTrack);
+		void RemoveVideoTrack(EpicRtcVideoTrackInterface* VideoTrack);
 
 		void		   OnStatsReady(const FString& PlayerId, const EpicRtcConnectionStats& ConnectionStats);
 		void		   OnFrameCapturerCreated();
