Fix: Crash when passing TextboxEntry from PixelStreaming - FocusedWidget can be nullptr sometimes, so using it should be skipped.

diff --git a/Engine/Plugins/Media/PixelStreaming2/Source/PixelStreaming2Input/Private/InputHandler.cpp b/Engine/Plugins/Media/PixelStreaming2/Source/PixelStreaming2Input/Private/InputHandler.cpp
--- a/Engine/Plugins/Media/PixelStreaming2/Source/PixelStreaming2Input/Private/InputHandler.cpp
+++ b/Engine/Plugins/Media/PixelStreaming2/Source/PixelStreaming2Input/Private/InputHandler.cpp
@@ -1104,6 +1104,11 @@ namespace UE::PixelStreaming2Input
 		FSlateApplication::Get().ForEachUser([this](FSlateUser& User) {
 			TSharedPtr<SWidget> FocusedWidget = User.GetFocusedWidget();
 
+			if(!FocusedWidget)
+			{
+				return;
+			}
+
 			static FName SEditableTextType(TEXT("SEditableText"));
 			static FName SMultiLineEditableTextType(TEXT("SMultiLineEditableText"));
 			bool		 bEditable = FocusedWidget && (FocusedWidget->GetType() == SEditableTextType || FocusedWidget->GetType() == SMultiLineEditableTextType);
