[PS2] Fix: Remove EpicRtc manager class as composition with it is causing lifetime issues

diff --git a/Engine/Plugins/Media/PixelStreaming2/Source/PixelStreaming2RTC/Internal/EpicRtcAudioTrackObserver.h b/Engine/Plugins/Media/PixelStreaming2/Source/PixelStreaming2RTC/Internal/EpicRtcAudioTrackObserver.h
--- a/Engine/Plugins/Media/PixelStreaming2/Source/PixelStreaming2RTC/Internal/EpicRtcAudioTrackObserver.h
+++ b/Engine/Plugins/Media/PixelStreaming2/Source/PixelStreaming2RTC/Internal/EpicRtcAudioTrackObserver.h
@@ -2,19 +2,36 @@
 
 #pragma once
 
-#include "Templates/SharedPointer.h"
+#include "EpicRtcObserver.h"
 
 #include "epic_rtc/core/audio/audio_track_observer.h"
 #include "epic_rtc_helper/memory/ref_count_impl_helper.h"
 
-namespace UE::PixelStreaming2
+#include "EpicRtcAudioTrackObserver.generated.h"
+
+UINTERFACE(MinimalAPI)
+class UPixelStreaming2AudioTrackObserver : public UInterface
+{
+	GENERATED_BODY()
+};
+
+class PIXELSTREAMING2RTC_API IPixelStreaming2AudioTrackObserver
 {
-	class FEpicRtcManager;
+	GENERATED_BODY()
 
+public:
+	virtual void OnAudioTrackMuted(EpicRtcAudioTrackInterface* AudioTrack, EpicRtcBool bIsMuted) = 0;
+	virtual void OnAudioTrackFrame(EpicRtcAudioTrackInterface* AudioTrack, const EpicRtcAudioFrame& Frame) = 0;
+	virtual void OnAudioTrackRemoved(EpicRtcAudioTrackInterface* AudioTrack) = 0;
+	virtual void OnAudioTrackState(EpicRtcAudioTrackInterface* AudioTrack, const EpicRtcTrackState State) = 0;
+};
+
+namespace UE::PixelStreaming2
+{
 	class PIXELSTREAMING2RTC_API FEpicRtcAudioTrackObserver : public EpicRtcAudioTrackObserverInterface
 	{
 	public:
-		FEpicRtcAudioTrackObserver(TWeakPtr<FEpicRtcManager> Manager);
+		FEpicRtcAudioTrackObserver(TObserverVariant<IPixelStreaming2AudioTrackObserver> UserObserver);
 		virtual ~FEpicRtcAudioTrackObserver() = default;
 
 	private:
@@ -31,6 +48,6 @@ namespace UE::PixelStreaming2
 		// End EpicRtcRefCountInterface
 
 	private:
-		TWeakPtr<FEpicRtcManager> Manager;
+		TObserverVariant<IPixelStreaming2AudioTrackObserver> UserObserver;
 	};
 } // namespace UE::PixelStreaming2
\ No newline at end of file
diff --git a/Engine/Plugins/Media/PixelStreaming2/Source/PixelStreaming2RTC/Internal/EpicRtcAudioTrackObserverFactory.h b/Engine/Plugins/Media/PixelStreaming2/Source/PixelStreaming2RTC/Internal/EpicRtcAudioTrackObserverFactory.h
--- a/Engine/Plugins/Media/PixelStreaming2/Source/PixelStreaming2RTC/Internal/EpicRtcAudioTrackObserverFactory.h
+++ b/Engine/Plugins/Media/PixelStreaming2/Source/PixelStreaming2RTC/Internal/EpicRtcAudioTrackObserverFactory.h
@@ -2,34 +2,33 @@
 
 #pragma once
 
-#include "EpicRtcAudioTrackObserver.h"
-#include "Templates/SharedPointer.h"
+#include "EpicRtcObserver.h"
 
 #include "epic_rtc/core/audio/audio_track_observer.h"
 #include "epic_rtc_helper/memory/ref_count_impl_helper.h"
 
-namespace UE::PixelStreaming2
-{
-	class FEpicRtcManager;
+class IPixelStreaming2AudioTrackObserver;
 
+namespace UE::PixelStreaming2
+{		
 	class PIXELSTREAMING2RTC_API FEpicRtcAudioTrackObserverFactory : public EpicRtcAudioTrackObserverFactoryInterface
 	{
 	public:
-		FEpicRtcAudioTrackObserverFactory(TWeakPtr<FEpicRtcManager> Manager);
+		FEpicRtcAudioTrackObserverFactory(TObserverVariant<IPixelStreaming2AudioTrackObserver> UserObserver);
 		virtual ~FEpicRtcAudioTrackObserverFactory() = default;
 
 	public:
-		// Begin EpicRtcAudioTrackObserverFactoryInterface 
+		// Begin EpicRtcAudioTrackObserverFactoryInterface
 		virtual EpicRtcErrorCode CreateAudioTrackObserver(const EpicRtcStringView ParticipantId, const EpicRtcStringView AudioTrackId, EpicRtcAudioTrackObserverInterface** OutAudioTrackObserver) override;
-		// End EpicRtcAudioTrackObserverFactoryInterface 
+		// End EpicRtcAudioTrackObserverFactoryInterface
 
 	public:
-		// Begin EpicRtcRefCountInterface 
+		// Begin EpicRtcRefCountInterface
 		EPICRTC_REFCOUNT_INTERFACE_IN_PLACE
-		// End EpicRtcRefCountInterface 
+		// End EpicRtcRefCountInterface
 
 	private:
-		TWeakPtr<FEpicRtcManager> Manager;
+		TObserverVariant<IPixelStreaming2AudioTrackObserver> UserObserver;
 	};
 
 } // namespace UE::PixelStreaming2
\ No newline at end of file
diff --git a/Engine/Plugins/Media/PixelStreaming2/Source/PixelStreaming2RTC/Internal/EpicRtcConferenceUtils.h b/Engine/Plugins/Media/PixelStreaming2/Source/PixelStreaming2RTC/Internal/EpicRtcConferenceUtils.h
--- a/Engine/Plugins/Media/PixelStreaming2/Source/PixelStreaming2RTC/Internal/EpicRtcConferenceUtils.h
+++ b/Engine/Plugins/Media/PixelStreaming2/Source/PixelStreaming2RTC/Internal/EpicRtcConferenceUtils.h
@@ -13,7 +13,7 @@ namespace UE::PixelStreaming2
 	class PIXELSTREAMING2RTC_API FEpicRtcTickConferenceTask : public FPixelStreamingTickableTask
 	{
 	public:
-		FEpicRtcTickConferenceTask(TRefCountPtr<EpicRtcConferenceInterface> EpicRtcConference, const FString& TaskName = TEXT("EpicRtcTickConferenceTask"))
+		FEpicRtcTickConferenceTask(TRefCountPtr<EpicRtcConferenceInterface>& EpicRtcConference, const FString& TaskName = TEXT("EpicRtcTickConferenceTask"))
 			: EpicRtcConference(EpicRtcConference)
 			, TaskName(TaskName)
 		{
diff --git a/Engine/Plugins/Media/PixelStreaming2/Source/PixelStreaming2RTC/Internal/EpicRtcDataTrackObserver.h b/Engine/Plugins/Media/PixelStreaming2/Source/PixelStreaming2RTC/Internal/EpicRtcDataTrackObserver.h
--- a/Engine/Plugins/Media/PixelStreaming2/Source/PixelStreaming2RTC/Internal/EpicRtcDataTrackObserver.h
+++ b/Engine/Plugins/Media/PixelStreaming2/Source/PixelStreaming2RTC/Internal/EpicRtcDataTrackObserver.h
@@ -2,26 +2,42 @@
 
 #pragma once
 
-#include "Templates/SharedPointer.h"
+#include "EpicRtcObserver.h"
 
 #include "epic_rtc/core/data_track_observer.h"
 #include "epic_rtc_helper/memory/ref_count_impl_helper.h"
 
-namespace UE::PixelStreaming2
+#include "EpicRtcDataTrackObserver.generated.h"
+
+UINTERFACE(MinimalAPI)
+class UPixelStreaming2DataTrackObserver : public UInterface
+{
+	GENERATED_BODY()
+};
+
+class PIXELSTREAMING2RTC_API IPixelStreaming2DataTrackObserver
 {
-	class FEpicRtcManager;
+	GENERATED_BODY()
 
+public:
+	virtual void OnDataTrackState(EpicRtcDataTrackInterface* DataTrack, const EpicRtcTrackState State) = 0;
+	virtual void OnDataTrackMessage(EpicRtcDataTrackInterface* DataTrack) = 0;
+	virtual void OnDataTrackError(EpicRtcDataTrackInterface*, const EpicRtcErrorCode) = 0;
+};
+
+namespace UE::PixelStreaming2
+{
 	class PIXELSTREAMING2RTC_API FEpicRtcDataTrackObserver : public EpicRtcDataTrackObserverInterface
 	{
 	public:
-		FEpicRtcDataTrackObserver(TWeakPtr<FEpicRtcManager> Manager);
+		FEpicRtcDataTrackObserver(TObserverVariant<IPixelStreaming2DataTrackObserver> UserObserver);
 		virtual ~FEpicRtcDataTrackObserver() = default;
 
 	private:
 		// Begin EpicRtcDataTrackObserverInterface
 		virtual void OnDataTrackState(EpicRtcDataTrackInterface* DataTrack, const EpicRtcTrackState State) override;
 		virtual void OnDataTrackMessage(EpicRtcDataTrackInterface* DataTrack) override;
-		virtual void OnDataTrackError(EpicRtcDataTrackInterface*, const EpicRtcErrorCode) override {};
+		virtual void OnDataTrackError(EpicRtcDataTrackInterface*, const EpicRtcErrorCode) override;
 		// End EpicRtcDataTrackObserverInterface
 
 	public:
@@ -30,7 +46,7 @@ namespace UE::PixelStreaming2
 		// End EpicRtcRefCountInterface
 
 	private:
-		TWeakPtr<FEpicRtcManager> Manager;
+		TObserverVariant<IPixelStreaming2DataTrackObserver> UserObserver;
 	};
 
 } // namespace UE::PixelStreaming2
\ No newline at end of file
diff --git a/Engine/Plugins/Media/PixelStreaming2/Source/PixelStreaming2RTC/Internal/EpicRtcDataTrackObserverFactory.h b/Engine/Plugins/Media/PixelStreaming2/Source/PixelStreaming2RTC/Internal/EpicRtcDataTrackObserverFactory.h
--- a/Engine/Plugins/Media/PixelStreaming2/Source/PixelStreaming2RTC/Internal/EpicRtcDataTrackObserverFactory.h
+++ b/Engine/Plugins/Media/PixelStreaming2/Source/PixelStreaming2RTC/Internal/EpicRtcDataTrackObserverFactory.h
@@ -2,34 +2,33 @@
 
 #pragma once
 
-#include "EpicRtcDataTrackObserver.h"
-#include "Templates/SharedPointer.h"
+#include "EpicRtcObserver.h"
 
 #include "epic_rtc/core/data_track_observer.h"
 #include "epic_rtc_helper/memory/ref_count_impl_helper.h"
 
+class IPixelStreaming2DataTrackObserver;
+
 namespace UE::PixelStreaming2
 {
-	class FEpicRtcManager;
-
 	class PIXELSTREAMING2RTC_API FEpicRtcDataTrackObserverFactory : public EpicRtcDataTrackObserverFactoryInterface
 	{
 	public:
-		FEpicRtcDataTrackObserverFactory(TWeakPtr<FEpicRtcManager> Manager);
+		FEpicRtcDataTrackObserverFactory(TObserverVariant<IPixelStreaming2DataTrackObserver> UserObserver);
 		virtual ~FEpicRtcDataTrackObserverFactory() = default;
 
 	public:
-		// Begin EpicRtcDataTrackObserverFactoryInterface 
+		// Begin EpicRtcDataTrackObserverFactoryInterface
 		virtual EpicRtcErrorCode CreateDataTrackObserver(const EpicRtcStringView ParticipantId, const EpicRtcStringView DataTrackId, EpicRtcDataTrackObserverInterface** OutDataTrackObserver) override;
-		// End EpicRtcDataTrackObserverFactoryInterface 
+		// End EpicRtcDataTrackObserverFactoryInterface
 
 	public:
-		// Begin EpicRtcRefCountInterface 
+		// Begin EpicRtcRefCountInterface
 		EPICRTC_REFCOUNT_INTERFACE_IN_PLACE
-		// End EpicRtcRefCountInterface 
+		// End EpicRtcRefCountInterface
 
 	private:
-		TWeakPtr<FEpicRtcManager> Manager;
+		TObserverVariant<IPixelStreaming2DataTrackObserver> UserObserver;
 	};
 
 } // namespace UE::PixelStreaming2
\ No newline at end of file
diff --git a/Engine/Plugins/Media/PixelStreaming2/Source/PixelStreaming2RTC/Internal/EpicRtcManager.h b/Engine/Plugins/Media/PixelStreaming2/Source/PixelStreaming2RTC/Internal/EpicRtcManager.h
deleted file mode 100644
diff --git a/Engine/Plugins/Media/PixelStreaming2/Source/PixelStreaming2RTC/Internal/EpicRtcObserver.h b/Engine/Plugins/Media/PixelStreaming2/Source/PixelStreaming2RTC/Internal/EpicRtcObserver.h
new file mode 100644
--- /dev/null
+++ b/Engine/Plugins/Media/PixelStreaming2/Source/PixelStreaming2RTC/Internal/EpicRtcObserver.h
@@ -0,0 +1,87 @@
+// Copyright Epic Games, Inc. All Rights Reserved.
+
+#pragma once
+
+#include "Misc/TVariant.h"
+#include "Templates/SharedPointer.h"
+#include "UObject/Interface.h"
+#include "UObject/WeakObjectPtr.h"
+#include "UObject/WeakObjectPtrTemplates.h"
+
+#include <type_traits>
+
+namespace UE::PixelStreaming2
+{
+	template <typename T>
+	class TObserverVariant
+	{
+	public:
+		TObserverVariant() = default;
+
+		TObserverVariant(TVariant<TYPE_OF_NULLPTR, TWeakPtr<T>, TWeakObjectPtr<T>> ObserverVariant)
+			: ObserverVariant(ObserverVariant)
+		{
+		}
+
+		T* operator->() const
+		{
+			switch (ObserverVariant.GetIndex())
+			{
+				case TVariant<TYPE_OF_NULLPTR, TWeakPtr<T>, TWeakObjectPtr<T>>::template IndexOfType<TWeakPtr<T>>():
+				{
+					if (TSharedPtr<T> PinnedUserObserver = ObserverVariant.template Get<TWeakPtr<T>>().Pin())
+					{
+						return PinnedUserObserver.Get();
+					}
+					break;
+				}
+				case TVariant<TYPE_OF_NULLPTR, TWeakPtr<T>, TWeakObjectPtr<T>>::template IndexOfType<TWeakObjectPtr<T>>():
+				{
+					TWeakObjectPtr<T> WeakObserver = ObserverVariant.template Get<TWeakObjectPtr<T>>();
+					if (WeakObserver.IsValid())
+					{
+						return Cast<T>(WeakObserver.Get());
+					}
+					break;
+				}
+				default:
+					checkNoEntry();
+			}
+
+			return nullptr;
+		}
+
+		operator bool() const
+		{
+			switch (ObserverVariant.GetIndex())
+			{
+				case TVariant<TYPE_OF_NULLPTR, TWeakPtr<T>, TWeakObjectPtr<T>>::template IndexOfType<TWeakPtr<T>>():
+				{
+					return ObserverVariant.template Get<TWeakPtr<T>>().IsValid();
+				}
+				case TVariant<TYPE_OF_NULLPTR, TWeakPtr<T>, TWeakObjectPtr<T>>::template IndexOfType<TWeakObjectPtr<T>>():
+				{
+					return ObserverVariant.template Get<TWeakObjectPtr<T>>().IsValid();
+				}
+				default:
+					checkNoEntry();
+			}
+			return false;
+		}
+
+	private:
+		TVariant<TYPE_OF_NULLPTR, TWeakPtr<T>, TWeakObjectPtr<T>> ObserverVariant;
+	};
+
+	template <typename T>
+	inline TObserverVariant<T> TObserver(TWeakPtr<T> WeakObserver)
+	{
+		return TObserverVariant<T>(TVariant<TYPE_OF_NULLPTR, TWeakPtr<T>, TWeakObjectPtr<T>>(TInPlaceType<TWeakPtr<T>>(), WeakObserver));
+	}
+
+	template <typename T>
+	inline TObserverVariant<T> TObserver(TWeakObjectPtr<T> WeakObserver)
+	{
+		return TObserverVariant<T>(TVariant<TYPE_OF_NULLPTR, TWeakPtr<T>, TWeakObjectPtr<T>>(TInPlaceType<TWeakObjectPtr<T>>(), WeakObserver));
+	}
+} // namespace UE::PixelStreaming2
\ No newline at end of file
diff --git a/Engine/Plugins/Media/PixelStreaming2/Source/PixelStreaming2RTC/Internal/EpicRtcRoomObserver.h b/Engine/Plugins/Media/PixelStreaming2/Source/PixelStreaming2RTC/Internal/EpicRtcRoomObserver.h
--- a/Engine/Plugins/Media/PixelStreaming2/Source/PixelStreaming2RTC/Internal/EpicRtcRoomObserver.h
+++ b/Engine/Plugins/Media/PixelStreaming2/Source/PixelStreaming2RTC/Internal/EpicRtcRoomObserver.h
@@ -2,19 +2,41 @@
 
 #pragma once
 
-#include "Templates/SharedPointer.h"
+#include "EpicRtcObserver.h"
 
 #include "epic_rtc/core/room_observer.h"
 #include "epic_rtc_helper/memory/ref_count_impl_helper.h"
 
-namespace UE::PixelStreaming2
+#include "EpicRtcRoomObserver.generated.h"
+
+UINTERFACE(MinimalAPI)
+class UPixelStreaming2RoomObserver : public UInterface
+{
+	GENERATED_BODY()
+};
+
+class PIXELSTREAMING2RTC_API IPixelStreaming2RoomObserver
 {
-	class FEpicRtcManager;
+	GENERATED_BODY()
 
+public:
+	virtual void							   OnRoomStateUpdate(const EpicRtcRoomState State) = 0;
+	virtual void							   OnRoomJoinedUpdate(EpicRtcParticipantInterface* Participant) = 0;
+	virtual void							   OnRoomLeftUpdate(const EpicRtcStringView ParticipantId) = 0;
+	virtual void							   OnAudioTrackUpdate(EpicRtcParticipantInterface* Participant, EpicRtcAudioTrackInterface* AudioTrack) = 0;
+	virtual void							   OnVideoTrackUpdate(EpicRtcParticipantInterface* Participant, EpicRtcVideoTrackInterface* VideoTrack) = 0;
+	virtual void							   OnDataTrackUpdate(EpicRtcParticipantInterface* Participant, EpicRtcDataTrackInterface* DataTrack) = 0;
+	[[nodiscard]] virtual EpicRtcSdpInterface* OnLocalSdpUpdate(EpicRtcParticipantInterface* Participant, EpicRtcSdpInterface* Sdp) = 0;
+	[[nodiscard]] virtual EpicRtcSdpInterface* OnRemoteSdpUpdate(EpicRtcParticipantInterface* Participant, EpicRtcSdpInterface* Sdp) = 0;
+	virtual void							   OnRoomErrorUpdate(const EpicRtcErrorCode Error) = 0;
+};
+
+namespace UE::PixelStreaming2
+{
 	class PIXELSTREAMING2RTC_API FEpicRtcRoomObserver : public EpicRtcRoomObserverInterface
 	{
 	public:
-		FEpicRtcRoomObserver(TWeakPtr<FEpicRtcManager> Manager);
+		FEpicRtcRoomObserver(TObserverVariant<IPixelStreaming2RoomObserver> UserObserver);
 		virtual ~FEpicRtcRoomObserver() = default;
 
 	private:
@@ -36,7 +58,7 @@ namespace UE::PixelStreaming2
 		// End EpicRtcRefCountInterface
 
 	private:
-		TWeakPtr<FEpicRtcManager> Manager;
+		TObserverVariant<IPixelStreaming2RoomObserver> UserObserver;
 	};
 
 } // namespace UE::PixelStreaming2
\ No newline at end of file
diff --git a/Engine/Plugins/Media/PixelStreaming2/Source/PixelStreaming2RTC/Internal/EpicRtcSessionObserver.h b/Engine/Plugins/Media/PixelStreaming2/Source/PixelStreaming2RTC/Internal/EpicRtcSessionObserver.h
--- a/Engine/Plugins/Media/PixelStreaming2/Source/PixelStreaming2RTC/Internal/EpicRtcSessionObserver.h
+++ b/Engine/Plugins/Media/PixelStreaming2/Source/PixelStreaming2RTC/Internal/EpicRtcSessionObserver.h
@@ -2,19 +2,35 @@
 
 #pragma once
 
-#include "Templates/SharedPointer.h"
+#include "EpicRtcObserver.h"
 
 #include "epic_rtc/core/session_observer.h"
 #include "epic_rtc_helper/memory/ref_count_impl_helper.h"
 
-namespace UE::PixelStreaming2
+#include "EpicRtcSessionObserver.generated.h"
+
+UINTERFACE(MinimalAPI)
+class UPixelStreaming2SessionObserver : public UInterface
+{
+	GENERATED_BODY()
+};
+
+class PIXELSTREAMING2RTC_API IPixelStreaming2SessionObserver
 {
-	class FEpicRtcManager;
+	GENERATED_BODY()
 
+public:
+	virtual void OnSessionStateUpdate(const EpicRtcSessionState State) = 0;
+	virtual void OnSessionErrorUpdate(const EpicRtcErrorCode Error) = 0;
+	virtual void OnSessionRoomsAvailableUpdate(EpicRtcStringArrayInterface* RoomsList) = 0;
+};
+
+namespace UE::PixelStreaming2
+{
 	class PIXELSTREAMING2RTC_API FEpicRtcSessionObserver : public EpicRtcSessionObserverInterface
 	{
 	public:
-		FEpicRtcSessionObserver(TWeakPtr<FEpicRtcManager> Manager);
+		FEpicRtcSessionObserver(TObserverVariant<IPixelStreaming2SessionObserver> UserObserver);
 		virtual ~FEpicRtcSessionObserver() = default;
 
 	private:
@@ -28,7 +44,7 @@ namespace UE::PixelStreaming2
 		EPICRTC_REFCOUNT_INTERFACE_IN_PLACE
 		// End EpicRtcRefCountInterface
 	private:
-		TWeakPtr<FEpicRtcManager> Manager;
+		TObserverVariant<IPixelStreaming2SessionObserver> UserObserver;
 	};
 
 } // namespace UE::PixelStreaming2
diff --git a/Engine/Plugins/Media/PixelStreaming2/Source/PixelStreaming2RTC/Internal/EpicRtcVideoTrackObserver.h b/Engine/Plugins/Media/PixelStreaming2/Source/PixelStreaming2RTC/Internal/EpicRtcVideoTrackObserver.h
--- a/Engine/Plugins/Media/PixelStreaming2/Source/PixelStreaming2RTC/Internal/EpicRtcVideoTrackObserver.h
+++ b/Engine/Plugins/Media/PixelStreaming2/Source/PixelStreaming2RTC/Internal/EpicRtcVideoTrackObserver.h
@@ -2,19 +2,38 @@
 
 #pragma once
 
-#include "Templates/SharedPointer.h"
+#include "EpicRtcObserver.h"
 
 #include "epic_rtc/core/video/video_track_observer.h"
 #include "epic_rtc_helper/memory/ref_count_impl_helper.h"
 
-namespace UE::PixelStreaming2
+#include "EpicRtcVideoTrackObserver.generated.h"
+
+UINTERFACE(MinimalAPI)
+class UPixelStreaming2VideoTrackObserver : public UInterface
+{
+	GENERATED_BODY()
+};
+
+class PIXELSTREAMING2RTC_API IPixelStreaming2VideoTrackObserver
 {
-	class FEpicRtcManager;
+	GENERATED_BODY()
 
+public:
+	virtual void		OnVideoTrackMuted(EpicRtcVideoTrackInterface* VideoTrack, EpicRtcBool bIsMuted) = 0;
+	virtual void		OnVideoTrackFrame(EpicRtcVideoTrackInterface* VideoTrack, const EpicRtcVideoFrame& Frame) = 0;
+	virtual void		OnVideoTrackRemoved(EpicRtcVideoTrackInterface* VideoTrack) = 0;
+	virtual void		OnVideoTrackState(EpicRtcVideoTrackInterface* VideoTrack, const EpicRtcTrackState State) = 0;
+	virtual void		OnVideoTrackEncodedFrame(EpicRtcVideoTrackInterface*, const EpicRtcEncodedVideoFrame&) = 0;
+	virtual EpicRtcBool Enabled() const = 0;
+};
+
+namespace UE::PixelStreaming2
+{
 	class PIXELSTREAMING2RTC_API FEpicRtcVideoTrackObserver : public EpicRtcVideoTrackObserverInterface
 	{
 	public:
-		FEpicRtcVideoTrackObserver(TWeakPtr<FEpicRtcManager> Manager);
+		FEpicRtcVideoTrackObserver(TObserverVariant<IPixelStreaming2VideoTrackObserver> UserObserver);
 		virtual ~FEpicRtcVideoTrackObserver() = default;
 
 	private:
@@ -23,8 +42,8 @@ namespace UE::PixelStreaming2
 		virtual void		OnVideoTrackFrame(EpicRtcVideoTrackInterface* VideoTrack, const EpicRtcVideoFrame& Frame) override;
 		virtual void		OnVideoTrackRemoved(EpicRtcVideoTrackInterface* VideoTrack) override;
 		virtual void		OnVideoTrackState(EpicRtcVideoTrackInterface* VideoTrack, const EpicRtcTrackState State) override;
-		virtual void		OnVideoTrackEncodedFrame(EpicRtcVideoTrackInterface*, const EpicRtcEncodedVideoFrame&) override {};
-		virtual EpicRtcBool Enabled() const override { return true; };
+		virtual void		OnVideoTrackEncodedFrame(EpicRtcVideoTrackInterface*, const EpicRtcEncodedVideoFrame&) override;
+		virtual EpicRtcBool Enabled() const override;
 		// End EpicRtcVideoTrackObserverInterface
 
 	public:
@@ -33,7 +52,7 @@ namespace UE::PixelStreaming2
 		// End EpicRtcRefCountInterface
 
 	private:
-		TWeakPtr<FEpicRtcManager> Manager;
+		TObserverVariant<IPixelStreaming2VideoTrackObserver> UserObserver;
 	};
 
 } // namespace UE::PixelStreaming2
\ No newline at end of file
diff --git a/Engine/Plugins/Media/PixelStreaming2/Source/PixelStreaming2RTC/Internal/EpicRtcVideoTrackObserverFactory.h b/Engine/Plugins/Media/PixelStreaming2/Source/PixelStreaming2RTC/Internal/EpicRtcVideoTrackObserverFactory.h
--- a/Engine/Plugins/Media/PixelStreaming2/Source/PixelStreaming2RTC/Internal/EpicRtcVideoTrackObserverFactory.h
+++ b/Engine/Plugins/Media/PixelStreaming2/Source/PixelStreaming2RTC/Internal/EpicRtcVideoTrackObserverFactory.h
@@ -2,19 +2,19 @@
 
 #pragma once
 
-#include "Templates/SharedPointer.h"
+#include "EpicRtcObserver.h"
 
 #include "epic_rtc/core/video/video_track_observer.h"
 #include "epic_rtc_helper/memory/ref_count_impl_helper.h"
 
-namespace UE::PixelStreaming2
-{
-	class FEpicRtcManager;
+class IPixelStreaming2VideoTrackObserver;
 
+namespace UE::PixelStreaming2
+{	
 	class PIXELSTREAMING2RTC_API FEpicRtcVideoTrackObserverFactory : public EpicRtcVideoTrackObserverFactoryInterface
 	{
 	public:
-		FEpicRtcVideoTrackObserverFactory(TWeakPtr<FEpicRtcManager> Manager);
+		FEpicRtcVideoTrackObserverFactory(TObserverVariant<IPixelStreaming2VideoTrackObserver> UserObserver);
 		virtual ~FEpicRtcVideoTrackObserverFactory() = default;
 
 	public:
@@ -28,7 +28,7 @@ namespace UE::PixelStreaming2
 		// End EpicRtcRefCountInterface
 
 	private:
-		TWeakPtr<FEpicRtcManager> Manager;
+		TObserverVariant<IPixelStreaming2VideoTrackObserver> UserObserver;
 	};
 
 } // namespace UE::PixelStreaming2
\ No newline at end of file
diff --git a/Engine/Plugins/Media/PixelStreaming2/Source/PixelStreaming2RTC/Private/Blueprints/PixelStreaming2Peer.cpp b/Engine/Plugins/Media/PixelStreaming2/Source/PixelStreaming2RTC/Private/Blueprints/PixelStreaming2Peer.cpp
--- a/Engine/Plugins/Media/PixelStreaming2/Source/PixelStreaming2RTC/Private/Blueprints/PixelStreaming2Peer.cpp
+++ b/Engine/Plugins/Media/PixelStreaming2/Source/PixelStreaming2RTC/Private/Blueprints/PixelStreaming2Peer.cpp
@@ -27,44 +27,13 @@ UPixelStreaming2Peer::UPixelStreaming2Peer(const FObjectInitializer& ObjectIniti
 
 void UPixelStreaming2Peer::BeginPlay()
 {
-	EpicRtcManager = MakeShared<UE::PixelStreaming2::FEpicRtcManager>();
+	const TWeakObjectPtr<UPixelStreaming2Peer> WeakThis = this;
+	SessionObserver = MakeRefCount<UE::PixelStreaming2::FEpicRtcSessionObserver>(UE::PixelStreaming2::TObserver<IPixelStreaming2SessionObserver>(WeakThis));
+	RoomObserver = MakeRefCount<UE::PixelStreaming2::FEpicRtcRoomObserver>(UE::PixelStreaming2::TObserver<IPixelStreaming2RoomObserver>(WeakThis));
 
-	EpicRtcManager->OnSessionStateUpdate.AddUObject(this, &UPixelStreaming2Peer::OnSessionStateUpdate);
-	EpicRtcManager->OnSessionErrorUpdate.AddUObject(this, &UPixelStreaming2Peer::OnSessionErrorUpdate);
-	EpicRtcManager->OnSessionRoomsAvailableUpdate.AddUObject(this, &UPixelStreaming2Peer::OnSessionRoomsAvailableUpdate);
-
-	EpicRtcManager->OnRoomStateUpdate.AddUObject(this, &UPixelStreaming2Peer::OnRoomStateUpdate);
-	EpicRtcManager->OnRoomJoinedUpdate.AddUObject(this, &UPixelStreaming2Peer::OnRoomJoinedUpdate);
-	EpicRtcManager->OnRoomLeftUpdate.AddUObject(this, &UPixelStreaming2Peer::OnRoomLeftUpdate);
-	EpicRtcManager->OnAudioTrackUpdate.AddUObject(this, &UPixelStreaming2Peer::OnAudioTrackUpdate);
-	EpicRtcManager->OnVideoTrackUpdate.AddUObject(this, &UPixelStreaming2Peer::OnVideoTrackUpdate);
-	EpicRtcManager->OnDataTrackUpdate.AddUObject(this, &UPixelStreaming2Peer::OnDataTrackUpdate);
-	EpicRtcManager->OnLocalSdpUpdate.AddUObject(this, &UPixelStreaming2Peer::OnLocalSdpUpdate);
-	EpicRtcManager->OnRemoteSdpUpdate.AddUObject(this, &UPixelStreaming2Peer::OnRemoteSdpUpdate);
-	EpicRtcManager->OnRoomErrorUpdate.AddUObject(this, &UPixelStreaming2Peer::OnRoomErrorUpdate);
-
-	EpicRtcManager->OnAudioTrackMuted.AddUObject(this, &UPixelStreaming2Peer::OnAudioTrackMuted);
-	EpicRtcManager->OnAudioTrackFrame.AddUObject(this, &UPixelStreaming2Peer::OnAudioTrackFrame);
-	EpicRtcManager->OnAudioTrackRemoved.AddUObject(this, &UPixelStreaming2Peer::OnAudioTrackRemoved);
-	EpicRtcManager->OnAudioTrackState.AddUObject(this, &UPixelStreaming2Peer::OnAudioTrackState);
-
-	EpicRtcManager->OnVideoTrackMuted.AddUObject(this, &UPixelStreaming2Peer::OnVideoTrackMuted);
-	EpicRtcManager->OnVideoTrackFrame.AddUObject(this, &UPixelStreaming2Peer::OnVideoTrackFrame);
-	EpicRtcManager->OnVideoTrackRemoved.AddUObject(this, &UPixelStreaming2Peer::OnVideoTrackRemoved);
-	EpicRtcManager->OnVideoTrackState.AddUObject(this, &UPixelStreaming2Peer::OnVideoTrackState);
-
-	EpicRtcManager->OnDataTrackRemoved.AddUObject(this, &UPixelStreaming2Peer::OnDataTrackRemoved);
-	EpicRtcManager->OnDataTrackState.AddUObject(this, &UPixelStreaming2Peer::OnDataTrackState);
-	EpicRtcManager->OnDataTrackMessage.AddUObject(this, &UPixelStreaming2Peer::OnDataTrackMessage);
-
-	EpicRtcManager->EpicRtcConference = UE::PixelStreaming2::FPixelStreaming2RTCModule::GetModule()->GetEpicRtcConference();
-
-	EpicRtcManager->SessionObserver = new UE::PixelStreaming2::FEpicRtcSessionObserver(EpicRtcManager);
-	EpicRtcManager->RoomObserver = new UE::PixelStreaming2::FEpicRtcRoomObserver(EpicRtcManager);
-
-	EpicRtcManager->AudioTrackObserverFactory = new UE::PixelStreaming2::FEpicRtcAudioTrackObserverFactory(EpicRtcManager);
-	EpicRtcManager->VideoTrackObserverFactory = new UE::PixelStreaming2::FEpicRtcVideoTrackObserverFactory(EpicRtcManager);
-	EpicRtcManager->DataTrackObserverFactory = new UE::PixelStreaming2::FEpicRtcDataTrackObserverFactory(EpicRtcManager);
+	AudioTrackObserverFactory = MakeRefCount<UE::PixelStreaming2::FEpicRtcAudioTrackObserverFactory>(UE::PixelStreaming2::TObserver<IPixelStreaming2AudioTrackObserver>(WeakThis));
+	VideoTrackObserverFactory = MakeRefCount<UE::PixelStreaming2::FEpicRtcVideoTrackObserverFactory>(UE::PixelStreaming2::TObserver<IPixelStreaming2VideoTrackObserver>(WeakThis));
+	DataTrackObserverFactory = MakeRefCount<UE::PixelStreaming2::FEpicRtcDataTrackObserverFactory>(UE::PixelStreaming2::TObserver<IPixelStreaming2DataTrackObserver>(WeakThis));
 
 	UE::PixelStreaming2::FPixelStreaming2RTCModule::GetModule()->GetStatsCollector()->OnStatsReady.AddUObject(this, &UPixelStreaming2Peer::OnStatsReady);
 
@@ -87,29 +56,23 @@ void UPixelStreaming2Peer::EndPlay(const EEndPlayReason::Type EndPlayReason)
 
 bool UPixelStreaming2Peer::Connect(const FString& Url)
 {
-	if (!EpicRtcManager)
-	{
-		UE_LOGFMT(LogPixelStreaming2RTC, Warning, "Failed to connect. EpicRtcManager isn't valid!");
-		return false;
-	}
-
 	FUtf8String Utf8Url = *Url;
 	FUtf8String ConnectionUrl = Utf8Url + (Utf8Url.Contains(TEXT("?")) ? TEXT("&") : TEXT("?")) + TEXT("isStreamer=false");
 
 	EpicRtcSessionConfig SessionConfig = {
 		._id = { ._ptr = reinterpret_cast<const char*>(*PlayerName), ._length = static_cast<uint64_t>(PlayerName.Len()) },
 		._url = { ._ptr = reinterpret_cast<const char*>(*ConnectionUrl), ._length = static_cast<uint64_t>(ConnectionUrl.Len()) },
-		._observer = EpicRtcManager->SessionObserver.GetReference()
+		._observer = SessionObserver.GetReference()
 	};
 
-	EpicRtcErrorCode Result = EpicRtcManager->EpicRtcConference->CreateSession(SessionConfig, EpicRtcManager->EpicRtcSession.GetInitReference());
+	EpicRtcErrorCode Result = UE::PixelStreaming2::FPixelStreaming2RTCModule::GetModule()->GetEpicRtcConference()->CreateSession(SessionConfig, EpicRtcSession.GetInitReference());
 	if (Result != EpicRtcErrorCode::Ok)
 	{
 		UE_LOGFMT(LogPixelStreaming2RTC, Error, "Failed to create EpicRtc session, CreateSession returned: {0}", UE::PixelStreaming2::ToString(Result));
 		return false;
 	}
 
-	Result = EpicRtcManager->EpicRtcSession->Connect();
+	Result = EpicRtcSession->Connect();
 	if (Result != EpicRtcErrorCode::Ok)
 	{
 		UE_LOGFMT(LogPixelStreaming2RTC, Error, "Failed to connect EpicRtcSession. Connect returned: {0}", UE::PixelStreaming2::ToString(Result));
@@ -125,22 +88,16 @@ bool UPixelStreaming2Peer::Disconnect()
 
 bool UPixelStreaming2Peer::Disconnect(const FString& OptionalReason)
 {
-	if (!EpicRtcManager)
-	{
-		UE_LOGFMT(LogPixelStreaming2RTC, Warning, "Failed to disconnect. EpicRtcManager isn't valid!");
-		return false;
-	}
-
-	if (!EpicRtcManager->EpicRtcSession)
+	if (!EpicRtcSession)
 	{
 		UE_LOGFMT(LogPixelStreaming2RTC, Error, "Failed to disconnect, EpicRtcSession does not exist");
 		return false;
 	}
 
-	if (EpicRtcManager->EpicRtcRoom)
+	if (EpicRtcRoom)
 	{
-		EpicRtcManager->EpicRtcRoom->Leave();
-		EpicRtcManager->EpicRtcSession->RemoveRoom({ ._ptr = reinterpret_cast<const char*>(*SubscribedStream), ._length = static_cast<uint64_t>(SubscribedStream.Len()) });
+		EpicRtcRoom->Leave();
+		EpicRtcSession->RemoveRoom({ ._ptr = reinterpret_cast<const char*>(*SubscribedStream), ._length = static_cast<uint64_t>(SubscribedStream.Len()) });
 	}
 
 	FUtf8String Reason;
@@ -153,7 +110,7 @@ bool UPixelStreaming2Peer::Disconnect(const FString& OptionalReason)
 		Reason = "PixelStreaming2Peer Disconnected";
 	}
 
-	EpicRtcErrorCode Result = EpicRtcManager->EpicRtcSession->Disconnect(UE::PixelStreaming2::ToEpicRtcStringView(Reason));
+	EpicRtcErrorCode Result = EpicRtcSession->Disconnect(UE::PixelStreaming2::ToEpicRtcStringView(Reason));
 	if (Result != EpicRtcErrorCode::Ok)
 	{
 		UE_LOGFMT(LogPixelStreaming2RTC, Error, "Failed to disconnect EpicRtcSession. Disconnect returned {0}", UE::PixelStreaming2::ToString(Result));
@@ -170,19 +127,13 @@ bool UPixelStreaming2Peer::Disconnect(const FString& OptionalReason)
 		VideoSink->RemoveVideoConsumer(VideoConsumer);
 	}
 
-	EpicRtcManager->EpicRtcConference->RemoveSession({ ._ptr = reinterpret_cast<const char*>(*PlayerName), ._length = static_cast<uint64_t>(PlayerName.Len()) });
+	UE::PixelStreaming2::FPixelStreaming2RTCModule::GetModule()->GetEpicRtcConference()->RemoveSession({ ._ptr = reinterpret_cast<const char*>(*PlayerName), ._length = static_cast<uint64_t>(PlayerName.Len()) });
 
 	return true;
 }
 
 bool UPixelStreaming2Peer::Subscribe(const FString& StreamerId)
 {
-	if (!EpicRtcManager)
-	{
-		UE_LOGFMT(LogPixelStreaming2RTC, Warning, "Failed to subscribe. EpicRtcManager isn't valid!");
-		return false;
-	}
-
 	if (SessionState != EpicRtcSessionState::Connected)
 	{
 		UE_LOGFMT(LogPixelStreaming2RTC, Error, "Failed to create subscribe to streamer. EpicRtc session isn't connected");
@@ -202,20 +153,20 @@ bool UPixelStreaming2Peer::Subscribe(const FString& StreamerId)
 		._id = EpicRtcStringView{ ._ptr = reinterpret_cast<const char*>(*SubscribedStream), ._length = static_cast<uint64_t>(SubscribedStream.Len()) },
 		._connectionConfig = ConnectionConfig,
 		._ticket = EpicRtcStringView{ ._ptr = nullptr, ._length = 0 },
-		._observer = EpicRtcManager->RoomObserver,
-		._audioTrackObserverFactory = EpicRtcManager->AudioTrackObserverFactory,
-		._dataTrackObserverFactory = EpicRtcManager->DataTrackObserverFactory,
-		._videoTrackObserverFactory = EpicRtcManager->VideoTrackObserverFactory
+		._observer = RoomObserver,
+		._audioTrackObserverFactory = AudioTrackObserverFactory,
+		._dataTrackObserverFactory = DataTrackObserverFactory,
+		._videoTrackObserverFactory = VideoTrackObserverFactory
 	};
 
-	EpicRtcErrorCode Result = EpicRtcManager->EpicRtcSession->CreateRoom(RoomConfig, EpicRtcManager->EpicRtcRoom.GetInitReference());
+	EpicRtcErrorCode Result = EpicRtcSession->CreateRoom(RoomConfig, EpicRtcRoom.GetInitReference());
 	if (Result != EpicRtcErrorCode::Ok)
 	{
 		UE_LOGFMT(LogPixelStreaming2RTC, Error, "Failed to create EpicRtc room. CreateRoom returned: {0}", UE::PixelStreaming2::ToString(Result));
 		return false;
 	}
 
-	EpicRtcManager->EpicRtcRoom->Join();
+	EpicRtcRoom->Join();
 
 	// Create a stats collector so we can receive stats from the subscribed streamer
 	StatsCollector = UE::PixelStreaming2::FRTCStatsCollector::Create(StreamerId);
@@ -366,12 +317,14 @@ void UPixelStreaming2Peer::OnDataTrackUpdate(EpicRtcParticipantInterface* Partic
 {
 }
 
-void UPixelStreaming2Peer::OnLocalSdpUpdate(EpicRtcParticipantInterface* Participant, EpicRtcSdpInterface* Sdp)
+[[nodiscard]] EpicRtcSdpInterface* UPixelStreaming2Peer::OnLocalSdpUpdate(EpicRtcParticipantInterface* Participant, EpicRtcSdpInterface* Sdp)
 {
+	return nullptr;
 }
 
-void UPixelStreaming2Peer::OnRemoteSdpUpdate(EpicRtcParticipantInterface* Participant, EpicRtcSdpInterface* Sdp)
+[[nodiscard]] EpicRtcSdpInterface* UPixelStreaming2Peer::OnRemoteSdpUpdate(EpicRtcParticipantInterface* Participant, EpicRtcSdpInterface* Sdp)
 {
+	return nullptr;
 }
 
 void UPixelStreaming2Peer::OnRoomErrorUpdate(const EpicRtcErrorCode Error)
@@ -434,10 +387,15 @@ void UPixelStreaming2Peer::OnVideoTrackState(EpicRtcVideoTrackInterface* VideoTr
 {
 }
 
-void UPixelStreaming2Peer::OnDataTrackRemoved(EpicRtcDataTrackInterface* DataTrack)
+void UPixelStreaming2Peer::OnVideoTrackEncodedFrame(EpicRtcVideoTrackInterface* VideoTrack, const EpicRtcEncodedVideoFrame& EncodedFrame)
 {
 }
 
+EpicRtcBool UPixelStreaming2Peer::Enabled() const
+{
+	return true;
+}
+
 void UPixelStreaming2Peer::OnDataTrackState(EpicRtcDataTrackInterface* DataTrack, const EpicRtcTrackState State)
 {
 }
@@ -446,6 +404,10 @@ void UPixelStreaming2Peer::OnDataTrackMessage(EpicRtcDataTrackInterface* DataTra
 {
 }
 
+void UPixelStreaming2Peer::OnDataTrackError(EpicRtcDataTrackInterface* DataTrack, const EpicRtcErrorCode Error)
+{
+}
+
 void UPixelStreaming2Peer::OnStatsReady(const FString& PeerId, const EpicRtcConnectionStats& ConnectionStats)
 {
 	if (!StatsCollector)
diff --git a/Engine/Plugins/Media/PixelStreaming2/Source/PixelStreaming2RTC/Private/Blueprints/PixelStreaming2Peer.h b/Engine/Plugins/Media/PixelStreaming2/Source/PixelStreaming2RTC/Private/Blueprints/PixelStreaming2Peer.h
--- a/Engine/Plugins/Media/PixelStreaming2/Source/PixelStreaming2RTC/Private/Blueprints/PixelStreaming2Peer.h
+++ b/Engine/Plugins/Media/PixelStreaming2/Source/PixelStreaming2RTC/Private/Blueprints/PixelStreaming2Peer.h
@@ -6,11 +6,22 @@
 #include "Components/SynthComponent.h"
 #include "Containers/Utf8String.h"
 #include "EpicRtcAudioSink.h"
-#include "EpicRtcManager.h"
+#include "EpicRtcAudioTrackObserver.h"
+#include "EpicRtcAudioTrackObserverFactory.h"
+#include "EpicRtcDataTrackObserver.h"
+#include "EpicRtcDataTrackObserverFactory.h"
 #include "EpicRtcVideoSink.h"
+#include "EpicRtcVideoTrackObserver.h"
+#include "EpicRtcVideoTrackObserverFactory.h"
+#include "EpicRtcRoomObserver.h"
+#include "EpicRtcSessionObserver.h"
 #include "IPixelStreaming2AudioConsumer.h"
 #include "RTCStatsCollector.h"
 
+#include "epic_rtc/core/conference.h"
+#include "epic_rtc/core/room.h"
+#include "epic_rtc/core/session.h"
+
 #include "PixelStreaming2Peer.generated.h"
 
 DECLARE_DYNAMIC_MULTICAST_SPARSE_DELEGATE_OneParam(FPixelStreamingStreamerList, UPixelStreaming2Peer, OnStreamerList, const TArray<FString>&, StreamerList);
@@ -26,7 +37,14 @@ namespace UE::PixelStreaming2
  */
 // UCLASS(Blueprintable, ClassGroup = (PixelStreaming2), meta = (BlueprintSpawnableComponent))
 UCLASS(Blueprintable, Category = "PixelStreaming2", META = (DisplayName = "PixelStreaming Peer Component", BlueprintSpawnableComponent))
-class UPixelStreaming2Peer : public USynthComponent, public IPixelStreaming2AudioConsumer
+class UPixelStreaming2Peer :
+	public USynthComponent,
+	public IPixelStreaming2AudioConsumer,
+	public IPixelStreaming2SessionObserver,
+	public IPixelStreaming2RoomObserver,
+	public IPixelStreaming2AudioTrackObserver,
+	public IPixelStreaming2DataTrackObserver,
+	public IPixelStreaming2VideoTrackObserver
 {
 	GENERATED_UCLASS_BODY()
 protected:
@@ -86,8 +104,6 @@ private:
 
 	TSharedPtr<UE::PixelStreaming2::FEpicRtcVideoSink> VideoSink;
 
-	TSharedPtr<UE::PixelStreaming2::FEpicRtcManager> EpicRtcManager;
-
 	TSharedPtr<UE::PixelStreaming2::FRTCStatsCollector> StatsCollector;
 
 	bool Disconnect(const FString& OptionalReason);
@@ -99,38 +115,61 @@ public:
 	virtual void OnConsumerRemoved() override;
 	// End IPixelStreaming2AudioConsumer Callbacks
 
-	// Begin FEpicRtcManager Callbacks
-	void OnSessionStateUpdate(const EpicRtcSessionState StateUpdate);
-	void OnSessionErrorUpdate(const EpicRtcErrorCode ErrorUpdate);
-	void OnSessionRoomsAvailableUpdate(EpicRtcStringArrayInterface* RoomsList);
-
-	void OnRoomStateUpdate(const EpicRtcRoomState State);
-	void OnRoomJoinedUpdate(EpicRtcParticipantInterface* Participant);
-	void OnRoomLeftUpdate(const EpicRtcStringView ParticipantId);
-	void OnAudioTrackUpdate(EpicRtcParticipantInterface* Participant, EpicRtcAudioTrackInterface* AudioTrack);
-	void OnVideoTrackUpdate(EpicRtcParticipantInterface* Participant, EpicRtcVideoTrackInterface* VideoTrack);
-	void OnDataTrackUpdate(EpicRtcParticipantInterface* Participant, EpicRtcDataTrackInterface* DataTrack);
-	void OnLocalSdpUpdate(EpicRtcParticipantInterface* Participant, EpicRtcSdpInterface* Sdp);
-	void OnRemoteSdpUpdate(EpicRtcParticipantInterface* Participant, EpicRtcSdpInterface* Sdp);
-	void OnRoomErrorUpdate(const EpicRtcErrorCode Error);
-
-	void OnAudioTrackMuted(EpicRtcAudioTrackInterface* AudioTrack, EpicRtcBool bIsMuted);
-	void OnAudioTrackFrame(EpicRtcAudioTrackInterface* AudioTrack, const EpicRtcAudioFrame& Frame);
-	void OnAudioTrackRemoved(EpicRtcAudioTrackInterface* AudioTrack);
-	void OnAudioTrackState(EpicRtcAudioTrackInterface* AudioTrack, const EpicRtcTrackState State);
-
-	void OnVideoTrackMuted(EpicRtcVideoTrackInterface* VideoTrack, EpicRtcBool bIsMuted);
-	void OnVideoTrackFrame(EpicRtcVideoTrackInterface* VideoTrack, const EpicRtcVideoFrame& Frame);
-	void OnVideoTrackRemoved(EpicRtcVideoTrackInterface* VideoTrack);
-	void OnVideoTrackState(EpicRtcVideoTrackInterface* VideoTrack, const EpicRtcTrackState State);
-
-	void OnDataTrackRemoved(EpicRtcDataTrackInterface* DataTrack);
-	void OnDataTrackState(EpicRtcDataTrackInterface* DataTrack, const EpicRtcTrackState State);
-	void OnDataTrackMessage(EpicRtcDataTrackInterface* DataTrack);
-	// End FEpicRtcManager Callbacks
+	// Begin IPixelStreaming2SessionObserver
+	virtual void OnSessionStateUpdate(const EpicRtcSessionState StateUpdate) override;
+	virtual void OnSessionErrorUpdate(const EpicRtcErrorCode ErrorUpdate) override;
+	virtual void OnSessionRoomsAvailableUpdate(EpicRtcStringArrayInterface* RoomsList) override;
+	// End IPixelStreaming2SessionObserver
+
+	// Begin IPixelStreaming2RoomObserver
+	virtual void							   OnRoomStateUpdate(const EpicRtcRoomState State) override;
+	virtual void							   OnRoomJoinedUpdate(EpicRtcParticipantInterface* Participant) override;
+	virtual void							   OnRoomLeftUpdate(const EpicRtcStringView ParticipantId) override;
+	virtual void							   OnAudioTrackUpdate(EpicRtcParticipantInterface* Participant, EpicRtcAudioTrackInterface* AudioTrack) override;
+	virtual void							   OnVideoTrackUpdate(EpicRtcParticipantInterface* Participant, EpicRtcVideoTrackInterface* VideoTrack) override;
+	virtual void							   OnDataTrackUpdate(EpicRtcParticipantInterface* Participant, EpicRtcDataTrackInterface* DataTrack) override;
+	[[nodiscard]] virtual EpicRtcSdpInterface* OnLocalSdpUpdate(EpicRtcParticipantInterface* Participant, EpicRtcSdpInterface* Sdp) override;
+	[[nodiscard]] virtual EpicRtcSdpInterface* OnRemoteSdpUpdate(EpicRtcParticipantInterface* Participant, EpicRtcSdpInterface* Sdp) override;
+	virtual void							   OnRoomErrorUpdate(const EpicRtcErrorCode Error) override;
+	// End IPixelStreaming2RoomObserver
+
+	// Begin IPixelStreaming2AudioTrackObserver
+	virtual void OnAudioTrackMuted(EpicRtcAudioTrackInterface* AudioTrack, EpicRtcBool bIsMuted) override;
+	virtual void OnAudioTrackFrame(EpicRtcAudioTrackInterface* AudioTrack, const EpicRtcAudioFrame& Frame) override;
+	virtual void OnAudioTrackRemoved(EpicRtcAudioTrackInterface* AudioTrack) override;
+	virtual void OnAudioTrackState(EpicRtcAudioTrackInterface* AudioTrack, const EpicRtcTrackState State) override;
+	// End IPixelStreaming2AudioTrackObserver
+
+	// Begin IPixelStreaming2VideoTrackObserver
+	virtual void		OnVideoTrackMuted(EpicRtcVideoTrackInterface* VideoTrack, EpicRtcBool bIsMuted) override;
+	virtual void		OnVideoTrackFrame(EpicRtcVideoTrackInterface* VideoTrack, const EpicRtcVideoFrame& Frame) override;
+	virtual void		OnVideoTrackRemoved(EpicRtcVideoTrackInterface* VideoTrack) override;
+	virtual void		OnVideoTrackState(EpicRtcVideoTrackInterface* VideoTrack, const EpicRtcTrackState State) override;
+	virtual void		OnVideoTrackEncodedFrame(EpicRtcVideoTrackInterface* VideoTrack, const EpicRtcEncodedVideoFrame& EncodedFrame) override;
+	virtual EpicRtcBool Enabled() const override;
+	// End IPixelStreaming2VideoTrackObserver
+
+	// Begin IPixelStreaming2DataTrackObserver
+	virtual void OnDataTrackState(EpicRtcDataTrackInterface* DataTrack, const EpicRtcTrackState State) override;
+	virtual void OnDataTrackMessage(EpicRtcDataTrackInterface* DataTrack) override;
+	virtual void OnDataTrackError(EpicRtcDataTrackInterface* DataTrack, const EpicRtcErrorCode Error) override;
+	// End IPixelStreaming2DataTrackObserver
 
 	void OnStatsReady(const FString& PeerId, const EpicRtcConnectionStats& ConnectionStats);
 
 private:
 	EpicRtcSessionState SessionState = EpicRtcSessionState::Disconnected;
+
+	// Begin EpicRtc Classes
+	TRefCountPtr<EpicRtcSessionInterface> EpicRtcSession;
+	TRefCountPtr<EpicRtcRoomInterface>	  EpicRtcRoom;
+	// End EpicRtc Classes
+
+	// Begin EpicRtc Observers
+	TRefCountPtr<UE::PixelStreaming2::FEpicRtcSessionObserver>			 SessionObserver;
+	TRefCountPtr<UE::PixelStreaming2::FEpicRtcRoomObserver>				 RoomObserver;
+	TRefCountPtr<UE::PixelStreaming2::FEpicRtcAudioTrackObserverFactory> AudioTrackObserverFactory;
+	TRefCountPtr<UE::PixelStreaming2::FEpicRtcVideoTrackObserverFactory> VideoTrackObserverFactory;
+	TRefCountPtr<UE::PixelStreaming2::FEpicRtcDataTrackObserverFactory>	 DataTrackObserverFactory;
+	// End EpicRtc Observers
 };
diff --git a/Engine/Plugins/Media/PixelStreaming2/Source/PixelStreaming2RTC/Private/EpicRtcAudioTrackObserver.cpp b/Engine/Plugins/Media/PixelStreaming2/Source/PixelStreaming2RTC/Private/EpicRtcAudioTrackObserver.cpp
--- a/Engine/Plugins/Media/PixelStreaming2/Source/PixelStreaming2RTC/Private/EpicRtcAudioTrackObserver.cpp
+++ b/Engine/Plugins/Media/PixelStreaming2/Source/PixelStreaming2RTC/Private/EpicRtcAudioTrackObserver.cpp
@@ -2,54 +2,43 @@
 
 #include "EpicRtcAudioTrackObserver.h"
 
-#include "EpicRtcManager.h"
-
 namespace UE::PixelStreaming2
 {
-
-	FEpicRtcAudioTrackObserver::FEpicRtcAudioTrackObserver(TWeakPtr<FEpicRtcManager> Manager)
-		: Manager(Manager)
+	FEpicRtcAudioTrackObserver::FEpicRtcAudioTrackObserver(TObserverVariant<IPixelStreaming2AudioTrackObserver> UserObserver)
+		: UserObserver(UserObserver)
 	{
 	}
 
 	void FEpicRtcAudioTrackObserver::OnAudioTrackMuted(EpicRtcAudioTrackInterface* AudioTrack, EpicRtcBool bIsMuted)
 	{
-		if (!Manager.IsValid())
+		if (UserObserver)
 		{
-			return;
+			UserObserver->OnAudioTrackMuted(AudioTrack, bIsMuted);
 		}
-
-		Manager.Pin()->OnAudioTrackMuted.Broadcast(AudioTrack, bIsMuted);
 	}
 
 	void FEpicRtcAudioTrackObserver::OnAudioTrackFrame(EpicRtcAudioTrackInterface* AudioTrack, const EpicRtcAudioFrame& Frame)
 	{
-		if (!Manager.IsValid())
+		if (UserObserver)
 		{
-			return;
+			UserObserver->OnAudioTrackFrame(AudioTrack, Frame);
 		}
-
-		Manager.Pin()->OnAudioTrackFrame.Broadcast(AudioTrack, Frame);
 	}
 
 	void FEpicRtcAudioTrackObserver::OnAudioTrackRemoved(EpicRtcAudioTrackInterface* AudioTrack)
 	{
-		if (!Manager.IsValid())
+		if (UserObserver)
 		{
-			return;
+			UserObserver->OnAudioTrackRemoved(AudioTrack);
 		}
-
-		Manager.Pin()->OnAudioTrackRemoved.Broadcast(AudioTrack);
 	}
 
 	void FEpicRtcAudioTrackObserver::OnAudioTrackState(EpicRtcAudioTrackInterface* AudioTrack, const EpicRtcTrackState State)
 	{
-		if (!Manager.IsValid())
+		if (UserObserver)
 		{
-			return;
+			UserObserver->OnAudioTrackState(AudioTrack, State);
 		}
-
-		Manager.Pin()->OnAudioTrackState.Broadcast(AudioTrack, State);
 	}
 
 } // namespace UE::PixelStreaming2
\ No newline at end of file
diff --git a/Engine/Plugins/Media/PixelStreaming2/Source/PixelStreaming2RTC/Private/EpicRtcAudioTrackObserverFactory.cpp b/Engine/Plugins/Media/PixelStreaming2/Source/PixelStreaming2RTC/Private/EpicRtcAudioTrackObserverFactory.cpp
--- a/Engine/Plugins/Media/PixelStreaming2/Source/PixelStreaming2RTC/Private/EpicRtcAudioTrackObserverFactory.cpp
+++ b/Engine/Plugins/Media/PixelStreaming2/Source/PixelStreaming2RTC/Private/EpicRtcAudioTrackObserverFactory.cpp
@@ -3,19 +3,18 @@
 #include "EpicRtcAudioTrackObserverFactory.h"
 
 #include "EpicRtcAudioTrackObserver.h"
-#include "EpicRtcManager.h"
 
 namespace UE::PixelStreaming2
 {
 
-	FEpicRtcAudioTrackObserverFactory::FEpicRtcAudioTrackObserverFactory(TWeakPtr<FEpicRtcManager> Manager)
-		: Manager(Manager)
+	FEpicRtcAudioTrackObserverFactory::FEpicRtcAudioTrackObserverFactory(TObserverVariant<IPixelStreaming2AudioTrackObserver> UserObserver)
+		: UserObserver(UserObserver)
 	{
 	}
 
 	EpicRtcErrorCode FEpicRtcAudioTrackObserverFactory::CreateAudioTrackObserver(const EpicRtcStringView ParticipantId, const EpicRtcStringView AudioTrackId, EpicRtcAudioTrackObserverInterface** OutAudioTrackObserver)
 	{
-		EpicRtcAudioTrackObserverInterface* AudioTrackObserver = new FEpicRtcAudioTrackObserver(Manager);
+		EpicRtcAudioTrackObserverInterface* AudioTrackObserver = new FEpicRtcAudioTrackObserver(UserObserver);
 		// Because the ptr was created with new, we need to call AddRef ourself (ms spec compliant)
 		AudioTrackObserver->AddRef();
 
diff --git a/Engine/Plugins/Media/PixelStreaming2/Source/PixelStreaming2RTC/Private/EpicRtcDataTrackObserver.cpp b/Engine/Plugins/Media/PixelStreaming2/Source/PixelStreaming2RTC/Private/EpicRtcDataTrackObserver.cpp
--- a/Engine/Plugins/Media/PixelStreaming2/Source/PixelStreaming2RTC/Private/EpicRtcDataTrackObserver.cpp
+++ b/Engine/Plugins/Media/PixelStreaming2/Source/PixelStreaming2RTC/Private/EpicRtcDataTrackObserver.cpp
@@ -2,34 +2,34 @@
 
 #include "EpicRtcDataTrackObserver.h"
 
-#include "EpicRtcManager.h"
-
 namespace UE::PixelStreaming2
 {
-
-	FEpicRtcDataTrackObserver::FEpicRtcDataTrackObserver(TWeakPtr<FEpicRtcManager> Manager)
-		: Manager(Manager)
+	FEpicRtcDataTrackObserver::FEpicRtcDataTrackObserver(TObserverVariant<IPixelStreaming2DataTrackObserver> UserObserver)
+		: UserObserver(UserObserver)
 	{
 	}
 
 	void FEpicRtcDataTrackObserver::OnDataTrackState(EpicRtcDataTrackInterface* DataTrack, const EpicRtcTrackState State)
 	{
-		if (!Manager.IsValid())
+		if (UserObserver)
 		{
-			return;
+			UserObserver->OnDataTrackState(DataTrack, State);
 		}
-
-		Manager.Pin()->OnDataTrackState.Broadcast(DataTrack, State);
 	}
 
 	void FEpicRtcDataTrackObserver::OnDataTrackMessage(EpicRtcDataTrackInterface* DataTrack)
 	{
-		if (!Manager.IsValid())
+		if (UserObserver)
 		{
-			return;
+			UserObserver->OnDataTrackMessage(DataTrack);
 		}
-
-		Manager.Pin()->OnDataTrackMessage.Broadcast(DataTrack);
 	}
 
+	void FEpicRtcDataTrackObserver::OnDataTrackError(EpicRtcDataTrackInterface* DataTrack, const EpicRtcErrorCode Error)
+	{
+		if (UserObserver)
+		{
+			UserObserver->OnDataTrackError(DataTrack, Error);
+		}
+	}
 } // namespace UE::PixelStreaming2
\ No newline at end of file
diff --git a/Engine/Plugins/Media/PixelStreaming2/Source/PixelStreaming2RTC/Private/EpicRtcDataTrackObserverFactory.cpp b/Engine/Plugins/Media/PixelStreaming2/Source/PixelStreaming2RTC/Private/EpicRtcDataTrackObserverFactory.cpp
--- a/Engine/Plugins/Media/PixelStreaming2/Source/PixelStreaming2RTC/Private/EpicRtcDataTrackObserverFactory.cpp
+++ b/Engine/Plugins/Media/PixelStreaming2/Source/PixelStreaming2RTC/Private/EpicRtcDataTrackObserverFactory.cpp
@@ -3,19 +3,17 @@
 #include "EpicRtcDataTrackObserverFactory.h"
 
 #include "EpicRtcDataTrackObserver.h"
-#include "EpicRtcManager.h"
 
 namespace UE::PixelStreaming2
 {
-
-	FEpicRtcDataTrackObserverFactory::FEpicRtcDataTrackObserverFactory(TWeakPtr<FEpicRtcManager> Manager)
-		: Manager(Manager)
+	FEpicRtcDataTrackObserverFactory::FEpicRtcDataTrackObserverFactory(TObserverVariant<IPixelStreaming2DataTrackObserver> UserObserver)
+		: UserObserver(UserObserver)
 	{
 	}
 
 	EpicRtcErrorCode FEpicRtcDataTrackObserverFactory::CreateDataTrackObserver(const EpicRtcStringView ParticipantId, const EpicRtcStringView DataTrackId, EpicRtcDataTrackObserverInterface** OutDataTrackObserver)
 	{
-		EpicRtcDataTrackObserverInterface* DataTrackObserver = new FEpicRtcDataTrackObserver(Manager);
+		EpicRtcDataTrackObserverInterface* DataTrackObserver = new FEpicRtcDataTrackObserver(UserObserver);
 		// Because the ptr was created with new, we need to call AddRef ourself (ms spec compliant)
 		DataTrackObserver->AddRef();
 
diff --git a/Engine/Plugins/Media/PixelStreaming2/Source/PixelStreaming2RTC/Private/EpicRtcRoomObserver.cpp b/Engine/Plugins/Media/PixelStreaming2/Source/PixelStreaming2RTC/Private/EpicRtcRoomObserver.cpp
--- a/Engine/Plugins/Media/PixelStreaming2/Source/PixelStreaming2RTC/Private/EpicRtcRoomObserver.cpp
+++ b/Engine/Plugins/Media/PixelStreaming2/Source/PixelStreaming2RTC/Private/EpicRtcRoomObserver.cpp
@@ -2,108 +2,85 @@
 
 #include "EpicRtcRoomObserver.h"
 
-#include "EpicRtcManager.h"
-
 namespace UE::PixelStreaming2
 {
-
-	FEpicRtcRoomObserver::FEpicRtcRoomObserver(TWeakPtr<FEpicRtcManager> Manager)
-		: Manager(Manager)
+	FEpicRtcRoomObserver::FEpicRtcRoomObserver(TObserverVariant<IPixelStreaming2RoomObserver> UserObserver)
+		: UserObserver(UserObserver)
 	{
 	}
 
 	void FEpicRtcRoomObserver::OnRoomStateUpdate(const EpicRtcRoomState State)
 	{
-		if (!Manager.IsValid())
+		if (UserObserver)
 		{
-			return;
+			UserObserver->OnRoomStateUpdate(State);
 		}
-
-		Manager.Pin()->OnRoomStateUpdate.Broadcast(State);
 	}
 
 	void FEpicRtcRoomObserver::OnRoomJoinedUpdate(EpicRtcParticipantInterface* Participant)
 	{
-		if (!Manager.IsValid())
+		if (UserObserver)
 		{
-			return;
+			UserObserver->OnRoomJoinedUpdate(Participant);
 		}
-
-		Manager.Pin()->OnRoomJoinedUpdate.Broadcast(Participant);
 	}
 
 	void FEpicRtcRoomObserver::OnRoomLeftUpdate(const EpicRtcStringView ParticipantId)
 	{
-		if (!Manager.IsValid())
+		if (UserObserver)
 		{
-			return;
+			UserObserver->OnRoomLeftUpdate(ParticipantId);
 		}
-
-		Manager.Pin()->OnRoomLeftUpdate.Broadcast(ParticipantId);
 	}
 
 	void FEpicRtcRoomObserver::OnAudioTrackUpdate(EpicRtcParticipantInterface* Participant, EpicRtcAudioTrackInterface* AudioTrack)
 	{
-		if (!Manager.IsValid())
+		if (UserObserver)
 		{
-			return;
+			UserObserver->OnAudioTrackUpdate(Participant, AudioTrack);
 		}
-
-		Manager.Pin()->OnAudioTrackUpdate.Broadcast(Participant, AudioTrack);
 	}
 
 	void FEpicRtcRoomObserver::OnVideoTrackUpdate(EpicRtcParticipantInterface* Participant, EpicRtcVideoTrackInterface* VideoTrack)
 	{
-		if (!Manager.IsValid())
+		if (UserObserver)
 		{
-			return;
+			UserObserver->OnVideoTrackUpdate(Participant, VideoTrack);
 		}
-
-		Manager.Pin()->OnVideoTrackUpdate.Broadcast(Participant, VideoTrack);
 	}
 
 	void FEpicRtcRoomObserver::OnDataTrackUpdate(EpicRtcParticipantInterface* Participant, EpicRtcDataTrackInterface* DataTrack)
 	{
-		if (!Manager.IsValid())
+		if (UserObserver)
 		{
-			return;
+			UserObserver->OnDataTrackUpdate(Participant, DataTrack);
 		}
-
-		Manager.Pin()->OnDataTrackUpdate.Broadcast(Participant, DataTrack);
 	}
 
 	[[nodiscard]] EpicRtcSdpInterface* FEpicRtcRoomObserver::OnLocalSdpUpdate(EpicRtcParticipantInterface* Participant, EpicRtcSdpInterface* Sdp)
 	{
-		if (!Manager.IsValid())
+		if (UserObserver)
 		{
-			return nullptr;
+			return UserObserver->OnLocalSdpUpdate(Participant, Sdp);
 		}
-
-		Manager.Pin()->OnLocalSdpUpdate.Broadcast(Participant, Sdp);
-
-		return nullptr;
+		return Sdp;
 	}
 
 	[[nodiscard]] EpicRtcSdpInterface* FEpicRtcRoomObserver::OnRemoteSdpUpdate(EpicRtcParticipantInterface* Participant, EpicRtcSdpInterface* Sdp)
 	{
-		if (!Manager.IsValid())
+		if (UserObserver)
 		{
-			return nullptr;
+			return UserObserver->OnRemoteSdpUpdate(Participant, Sdp);
 		}
-
-		Manager.Pin()->OnRemoteSdpUpdate.Broadcast(Participant, Sdp);
-
-		return nullptr;
+		return Sdp;
 	}
 
 	void FEpicRtcRoomObserver::OnRoomErrorUpdate(const EpicRtcErrorCode Error)
 	{
-		if (!Manager.IsValid())
+		if (UserObserver)
 		{
-			return;
+			return UserObserver->OnRoomErrorUpdate(Error);
 		}
-
-		Manager.Pin()->OnRoomErrorUpdate.Broadcast(Error);
 	}
 
 } // namespace UE::PixelStreaming2
\ No newline at end of file
diff --git a/Engine/Plugins/Media/PixelStreaming2/Source/PixelStreaming2RTC/Private/EpicRtcSessionObserver.cpp b/Engine/Plugins/Media/PixelStreaming2/Source/PixelStreaming2RTC/Private/EpicRtcSessionObserver.cpp
--- a/Engine/Plugins/Media/PixelStreaming2/Source/PixelStreaming2RTC/Private/EpicRtcSessionObserver.cpp
+++ b/Engine/Plugins/Media/PixelStreaming2/Source/PixelStreaming2RTC/Private/EpicRtcSessionObserver.cpp
@@ -2,44 +2,36 @@
 
 #include "EpicRtcSessionObserver.h"
 
-#include "EpicRtcManager.h"
-
 namespace UE::PixelStreaming2
 {
 
-	FEpicRtcSessionObserver::FEpicRtcSessionObserver(TWeakPtr<FEpicRtcManager> Manager)
-		: Manager(Manager)
+	FEpicRtcSessionObserver::FEpicRtcSessionObserver(TObserverVariant<IPixelStreaming2SessionObserver> UserObserver)
+		: UserObserver(UserObserver)
 	{
 	}
 
 	void FEpicRtcSessionObserver::OnSessionStateUpdate(const EpicRtcSessionState State)
 	{
-		if (!Manager.IsValid())
+		if (UserObserver)
 		{
-			return;
+			return UserObserver->OnSessionStateUpdate(State);
 		}
-
-		Manager.Pin()->OnSessionStateUpdate.Broadcast(State);
 	}
 
 	void FEpicRtcSessionObserver::OnSessionErrorUpdate(const EpicRtcErrorCode Error)
 	{
-		if (!Manager.IsValid())
+		if (UserObserver)
 		{
-			return;
+			return UserObserver->OnSessionErrorUpdate(Error);
 		}
-
-		Manager.Pin()->OnSessionErrorUpdate.Broadcast(Error);
 	}
 
 	void FEpicRtcSessionObserver::OnSessionRoomsAvailableUpdate(EpicRtcStringArrayInterface* RoomsList)
 	{
-		if (!Manager.IsValid())
+		if (UserObserver)
 		{
-			return;
+			return UserObserver->OnSessionRoomsAvailableUpdate(RoomsList);
 		}
-
-		Manager.Pin()->OnSessionRoomsAvailableUpdate.Broadcast(RoomsList);
 	}
 
 } // namespace UE::PixelStreaming2
diff --git a/Engine/Plugins/Media/PixelStreaming2/Source/PixelStreaming2RTC/Private/EpicRtcStreamer.cpp b/Engine/Plugins/Media/PixelStreaming2/Source/PixelStreaming2RTC/Private/EpicRtcStreamer.cpp
--- a/Engine/Plugins/Media/PixelStreaming2/Source/PixelStreaming2RTC/Private/EpicRtcStreamer.cpp
+++ b/Engine/Plugins/Media/PixelStreaming2/Source/PixelStreaming2RTC/Private/EpicRtcStreamer.cpp
@@ -33,17 +33,17 @@
 #include "VideoProducerMediaCapture.h"
 #include "RTCInputHandler.h"
 #include "PixelStreaming2Common.h"
+#include "EpicRtcAudioCapturer.h"
 
 namespace UE::PixelStreaming2
 {
-	FEpicRtcStreamer::FEpicRtcStreamer(const FString& InStreamerId, TRefCountPtr<EpicRtcConferenceInterface> Conference)
+	FEpicRtcStreamer::FEpicRtcStreamer(const FString& InStreamerId)
 		: StreamerId(InStreamerId)
 		, InputHandler(FRTCInputHandler::Create())
 		, Players(new TThreadSafeMap<FString, FPlayerContext>())
 		, VideoCapturer(FEpicRtcVideoCapturer::Create())
 		, VideoSourceGroup(FVideoSourceGroup::Create(VideoCapturer))
 		, FreezeFrame(FFreezeFrame::Create(Players, VideoCapturer, InputHandler))
-		, EpicRtcManager(MakeShared<FEpicRtcManager>())
 	{
 		ReconnectTimer = MakeShared<FStreamerReconnectTimer>();
 
@@ -52,15 +52,6 @@ namespace UE::PixelStreaming2
 				|| InputControllingId == INVALID_PLAYER_ID
 				|| PlayerId == InputControllingId;
 		});
-
-		EpicRtcManager->EpicRtcConference = Conference;
-
-		EpicRtcManager->SessionObserver = MakeRefCount<FEpicRtcSessionObserver>(EpicRtcManager);
-		EpicRtcManager->RoomObserver = MakeRefCount<FEpicRtcRoomObserver>(EpicRtcManager);
-
-		EpicRtcManager->AudioTrackObserverFactory = MakeRefCount<FEpicRtcAudioTrackObserverFactory>(EpicRtcManager);
-		EpicRtcManager->VideoTrackObserverFactory = MakeRefCount<FEpicRtcVideoTrackObserverFactory>(EpicRtcManager);
-		EpicRtcManager->DataTrackObserverFactory = MakeRefCount<FEpicRtcDataTrackObserverFactory>(EpicRtcManager);
 	}
 
 	FEpicRtcStreamer::~FEpicRtcStreamer()
@@ -70,13 +61,13 @@ namespace UE::PixelStreaming2
 		// Pin ref counted objects and clear the session if it is still alive.
 		// Session may still be alive if the manager was destroyed quickly after stopping the stream
 		// and the EpicRtcSessionState::Disconnected event did not have enough time to tick.
-		const TRefCountPtr<EpicRtcConferenceInterface> Conference = EpicRtcManager->EpicRtcConference;
-		const TRefCountPtr<EpicRtcSessionInterface> Session = EpicRtcManager->EpicRtcSession;
+		const TRefCountPtr<EpicRtcConferenceInterface> Conference = FPixelStreaming2RTCModule::GetModule()->GetEpicRtcConference();
+		const TRefCountPtr<EpicRtcSessionInterface>	   Session = EpicRtcSession;
 		if (Conference.IsValid() && Session.IsValid())
 		{
 			const FUtf8String Utf8StreamerId(StreamerId);
 			Conference->RemoveSession(ToEpicRtcStringView(Utf8StreamerId));
-			EpicRtcManager->EpicRtcSession = nullptr;
+			EpicRtcSession = nullptr;
 		}
 	}
 
@@ -85,6 +76,13 @@ namespace UE::PixelStreaming2
 		TSharedPtr<FEpicRtcStreamer> Streamer = AsShared();
 		TWeakPtr<FEpicRtcStreamer>	 WeakStreamer = Streamer;
 
+		SessionObserver = MakeRefCount<FEpicRtcSessionObserver>(TObserver<IPixelStreaming2SessionObserver>(WeakStreamer));
+		RoomObserver = MakeRefCount<FEpicRtcRoomObserver>(TObserver<IPixelStreaming2RoomObserver>(WeakStreamer));
+
+		AudioTrackObserverFactory = MakeRefCount<FEpicRtcAudioTrackObserverFactory>(TObserver<IPixelStreaming2AudioTrackObserver>(WeakStreamer));
+		VideoTrackObserverFactory = MakeRefCount<FEpicRtcVideoTrackObserverFactory>(TObserver<IPixelStreaming2VideoTrackObserver>(WeakStreamer));
+		DataTrackObserverFactory = MakeRefCount<FEpicRtcDataTrackObserverFactory>(TObserver<IPixelStreaming2DataTrackObserver>(WeakStreamer));
+
 		InputHandler->GetToStreamerProtocol()->OnProtocolUpdated().AddSP(Streamer.ToSharedRef(), &FEpicRtcStreamer::OnProtocolUpdated);
 		InputHandler->GetFromStreamerProtocol()->OnProtocolUpdated().AddSP(Streamer.ToSharedRef(), &FEpicRtcStreamer::OnProtocolUpdated);
 
@@ -173,34 +171,6 @@ namespace UE::PixelStreaming2
 		// Some special cases include when virtual gamepads are connected and a controller id needs to be transmitted.
 		InputHandler->OnSendMessage.AddSP(Streamer.ToSharedRef(), &FEpicRtcStreamer::OnSendMessage);
 
-		Streamer->EpicRtcManager->OnSessionStateUpdate.AddSP(Streamer.ToSharedRef(), &FEpicRtcStreamer::OnSessionStateUpdate);
-		Streamer->EpicRtcManager->OnSessionErrorUpdate.AddSP(Streamer.ToSharedRef(), &FEpicRtcStreamer::OnSessionErrorUpdate);
-		Streamer->EpicRtcManager->OnSessionRoomsAvailableUpdate.AddSP(Streamer.ToSharedRef(), &FEpicRtcStreamer::OnSessionRoomsAvailableUpdate);
-
-		Streamer->EpicRtcManager->OnRoomStateUpdate.AddSP(Streamer.ToSharedRef(), &FEpicRtcStreamer::OnRoomStateUpdate);
-		Streamer->EpicRtcManager->OnRoomJoinedUpdate.AddSP(Streamer.ToSharedRef(), &FEpicRtcStreamer::OnRoomJoinedUpdate);
-		Streamer->EpicRtcManager->OnRoomLeftUpdate.AddSP(Streamer.ToSharedRef(), &FEpicRtcStreamer::OnRoomLeftUpdate);
-		Streamer->EpicRtcManager->OnAudioTrackUpdate.AddSP(Streamer.ToSharedRef(), &FEpicRtcStreamer::OnAudioTrackUpdate);
-		Streamer->EpicRtcManager->OnVideoTrackUpdate.AddSP(Streamer.ToSharedRef(), &FEpicRtcStreamer::OnVideoTrackUpdate);
-		Streamer->EpicRtcManager->OnDataTrackUpdate.AddSP(Streamer.ToSharedRef(), &FEpicRtcStreamer::OnDataTrackUpdate);
-		Streamer->EpicRtcManager->OnLocalSdpUpdate.AddSP(Streamer.ToSharedRef(), &FEpicRtcStreamer::OnLocalSdpUpdate);
-		Streamer->EpicRtcManager->OnRemoteSdpUpdate.AddSP(Streamer.ToSharedRef(), &FEpicRtcStreamer::OnRemoteSdpUpdate);
-		Streamer->EpicRtcManager->OnRoomErrorUpdate.AddSP(Streamer.ToSharedRef(), &FEpicRtcStreamer::OnRoomErrorUpdate);
-
-		Streamer->EpicRtcManager->OnAudioTrackMuted.AddSP(Streamer.ToSharedRef(), &FEpicRtcStreamer::OnAudioTrackMuted);
-		Streamer->EpicRtcManager->OnAudioTrackFrame.AddSP(Streamer.ToSharedRef(), &FEpicRtcStreamer::OnAudioTrackFrame);
-		Streamer->EpicRtcManager->OnAudioTrackRemoved.AddSP(Streamer.ToSharedRef(), &FEpicRtcStreamer::OnAudioTrackRemoved);
-		Streamer->EpicRtcManager->OnAudioTrackState.AddSP(Streamer.ToSharedRef(), &FEpicRtcStreamer::OnAudioTrackState);
-
-		Streamer->EpicRtcManager->OnVideoTrackMuted.AddSP(Streamer.ToSharedRef(), &FEpicRtcStreamer::OnVideoTrackMuted);
-		Streamer->EpicRtcManager->OnVideoTrackFrame.AddSP(Streamer.ToSharedRef(), &FEpicRtcStreamer::OnVideoTrackFrame);
-		Streamer->EpicRtcManager->OnVideoTrackRemoved.AddSP(Streamer.ToSharedRef(), &FEpicRtcStreamer::OnVideoTrackRemoved);
-		Streamer->EpicRtcManager->OnVideoTrackState.AddSP(Streamer.ToSharedRef(), &FEpicRtcStreamer::OnVideoTrackState);
-
-		Streamer->EpicRtcManager->OnDataTrackRemoved.AddSP(Streamer.ToSharedRef(), &FEpicRtcStreamer::OnDataTrackRemoved);
-		Streamer->EpicRtcManager->OnDataTrackState.AddSP(Streamer.ToSharedRef(), &FEpicRtcStreamer::OnDataTrackState);
-		Streamer->EpicRtcManager->OnDataTrackMessage.AddSP(Streamer.ToSharedRef(), &FEpicRtcStreamer::OnDataTrackMessage);
-
 		VideoCapturer->OnFrameCapturerCreated.AddSP(Streamer.ToSharedRef(), &FEpicRtcStreamer::OnFrameCapturerCreated);
 
 		if (UPixelStreaming2PluginSettings::FDelegates* Delegates = UPixelStreaming2PluginSettings::Delegates())
@@ -289,10 +259,10 @@ namespace UE::PixelStreaming2
 		EpicRtcSessionConfig SessionConfig{
 			._id = ToEpicRtcStringView(Utf8StreamerId),
 			._url = ToEpicRtcStringView(Utf8CurrentSignallingServerURL),
-			._observer = EpicRtcManager->SessionObserver
+			._observer = SessionObserver
 		};
 
-		if (const EpicRtcErrorCode Result = EpicRtcManager->EpicRtcConference->CreateSession(SessionConfig, EpicRtcManager->EpicRtcSession.GetInitReference()); Result != EpicRtcErrorCode::Ok)
+		if (const EpicRtcErrorCode Result = FPixelStreaming2RTCModule::GetModule()->GetEpicRtcConference()->CreateSession(SessionConfig, EpicRtcSession.GetInitReference()); Result != EpicRtcErrorCode::Ok)
 		{
 			UE_LOGFMT(LogPixelStreaming2RTC, Error, "Failed to create EpicRtc session. CreateSession returned {0}", *ToString(Result));
 			StreamState = EStreamState::Disconnected;
@@ -300,7 +270,7 @@ namespace UE::PixelStreaming2
 		}
 
 		// TODO (william.belcher): This should move to OnSessionStateUpdate(EpicRtcSessionState::New) once EpicRtc starts broadcasting that state
-		if (const EpicRtcErrorCode Result = EpicRtcManager->EpicRtcSession->Connect(); Result != EpicRtcErrorCode::Ok)
+		if (const EpicRtcErrorCode Result = EpicRtcSession->Connect(); Result != EpicRtcErrorCode::Ok)
 		{
 			UE_LOGFMT(LogPixelStreaming2RTC, Error, "Failed to connect EpicRtcSession. Connect returned {0}", ToString(Result));
 			StreamState = EStreamState::Disconnected;
@@ -317,7 +287,7 @@ namespace UE::PixelStreaming2
 
 		StreamState = EStreamState::Disconnecting;
 
-		if (const TRefCountPtr<EpicRtcRoomInterface> Room = EpicRtcManager->EpicRtcRoom)
+		if (const TRefCountPtr<EpicRtcRoomInterface> Room = EpicRtcRoom)
 		{
 			Room->Leave();
 		}
@@ -1035,13 +1005,13 @@ namespace UE::PixelStreaming2
 					._id = ToEpicRtcStringView(Utf8StreamerId),
 					._connectionConfig = ConnectionConfig,
 					._ticket = EpicRtcStringView{ ._ptr = nullptr, ._length = 0 },
-					._observer = EpicRtcManager->RoomObserver,
-					._audioTrackObserverFactory = EpicRtcManager->AudioTrackObserverFactory,
-					._dataTrackObserverFactory = EpicRtcManager->DataTrackObserverFactory,
-					._videoTrackObserverFactory = EpicRtcManager->VideoTrackObserverFactory
+					._observer = RoomObserver,
+					._audioTrackObserverFactory = AudioTrackObserverFactory,
+					._dataTrackObserverFactory = DataTrackObserverFactory,
+					._videoTrackObserverFactory = VideoTrackObserverFactory
 				};
 
-				EpicRtcErrorCode Result = EpicRtcManager->EpicRtcSession->CreateRoom(RoomConfig, EpicRtcManager->EpicRtcRoom.GetInitReference());
+				EpicRtcErrorCode Result = EpicRtcSession->CreateRoom(RoomConfig, EpicRtcRoom.GetInitReference());
 				if (Result != EpicRtcErrorCode::Ok)
 				{
 					UE_LOGFMT(LogPixelStreaming2RTC, Error, "Failed to create EpicRtc room. CreateRoom returned {0}", ToString(Result));
@@ -1049,7 +1019,7 @@ namespace UE::PixelStreaming2
 				}
 
 				// TODO (william.belcher): This should move to OnRoomStateUpdate(EpicRtcRoomState::New) once EpicRtc starts broadcasting that state
-				EpicRtcManager->EpicRtcRoom->Join();
+				EpicRtcRoom->Join();
 				break;
 			}
 			case EpicRtcSessionState::Disconnected: // Indicates session is disconnected from the signalling server.
@@ -1075,11 +1045,11 @@ namespace UE::PixelStreaming2
 
 				// The session has been disconnected (either through a call to StopStreaming or by an error)
 				// so remove it
-				if (EpicRtcManager->EpicRtcConference)
+				if (FPixelStreaming2RTCModule::GetModule()->GetEpicRtcConference())
 				{
 					FUtf8String Utf8StreamerId(StreamerId);
-					EpicRtcManager->EpicRtcConference->RemoveSession(ToEpicRtcStringView(Utf8StreamerId));
-					EpicRtcManager->EpicRtcSession = nullptr;
+					FPixelStreaming2RTCModule::GetModule()->GetEpicRtcConference()->RemoveSession(ToEpicRtcStringView(Utf8StreamerId));
+					EpicRtcSession = nullptr;
 				}
 
 				// We are fully disconnect at this point so we update the state so we can StartStreaming again
@@ -1144,16 +1114,16 @@ namespace UE::PixelStreaming2
 			{
 				UE_LOGFMT(LogPixelStreaming2RTC, VeryVerbose, "FEpicRtcStreamer::OnRoomStateUpdate State=Left");
 
-				if (!EpicRtcManager->EpicRtcSession)
+				if (!EpicRtcSession)
 				{
 					break;
 				}
 
 				FUtf8String Utf8StreamerId(StreamerId);
-				EpicRtcManager->EpicRtcSession->RemoveRoom(ToEpicRtcStringView(Utf8StreamerId));
-				EpicRtcManager->EpicRtcRoom = nullptr;
+				EpicRtcSession->RemoveRoom(ToEpicRtcStringView(Utf8StreamerId));
+				EpicRtcRoom = nullptr;
 
-				const EpicRtcErrorCode Result = EpicRtcManager->EpicRtcSession->Disconnect(ToEpicRtcStringView("Streaming Session Removed"));
+				const EpicRtcErrorCode Result = EpicRtcSession->Disconnect(ToEpicRtcStringView("Streaming Session Removed"));
 				if (Result == EpicRtcErrorCode::SessionDisconnected)
 				{
 					UE_LOGFMT(LogPixelStreaming2RTC, VeryVerbose, "Session disconnected cleanly.");
@@ -1472,7 +1442,7 @@ namespace UE::PixelStreaming2
 		}
 	}
 
-	void FEpicRtcStreamer::OnLocalSdpUpdate(EpicRtcParticipantInterface* Participant, EpicRtcSdpInterface* Sdp)
+	[[nodiscard]] EpicRtcSdpInterface* FEpicRtcStreamer::OnLocalSdpUpdate(EpicRtcParticipantInterface* Participant, EpicRtcSdpInterface* Sdp)
 	{
 		FString ParticipantId = ToString(Participant->GetId());
 		FString SdpType = TEXT("");
@@ -1485,10 +1455,13 @@ namespace UE::PixelStreaming2
 				SdpType = TEXT("Answer");
 				break;
 		}
+
 		UE_LOG(LogPixelStreaming2RTC, Log, TEXT("FEpicRtcStreamer::OnLocalSdpUpdate(Participant [%s], Type [%s])"), *ParticipantId, *SdpType);
+
+		return nullptr;
 	}
 
-	void FEpicRtcStreamer::OnRemoteSdpUpdate(EpicRtcParticipantInterface* Participant, EpicRtcSdpInterface* Sdp)
+	[[nodiscard]] EpicRtcSdpInterface* FEpicRtcStreamer::OnRemoteSdpUpdate(EpicRtcParticipantInterface* Participant, EpicRtcSdpInterface* Sdp)
 	{
 		FString ParticipantId = ToString(Participant->GetId());
 		FString SdpType = TEXT("");
@@ -1501,7 +1474,10 @@ namespace UE::PixelStreaming2
 				SdpType = TEXT("Answer");
 				break;
 		}
+
 		UE_LOG(LogPixelStreaming2RTC, Log, TEXT("FEpicRtcStreamer::OnRemoteSdpUpdate(Participant [%s], Type [%s])"), *ParticipantId, *SdpType);
+
+		return nullptr;
 	}
 
 	void FEpicRtcStreamer::OnRoomErrorUpdate(const EpicRtcErrorCode Error)
@@ -1554,7 +1530,7 @@ namespace UE::PixelStreaming2
 
 	void FEpicRtcStreamer::OnAudioTrackRemoved(EpicRtcAudioTrackInterface* AudioTrack)
 	{
-		// TODO (Eden.Harris) OnAudioTrackRemoved is not hooked up in EpicRtc. 
+		// TODO (Eden.Harris) OnAudioTrackRemoved is not hooked up in EpicRtc.
 		UE_LOGFMT(LogPixelStreaming2RTC, Warning, "FEpicRtcStreamer::OnAudioTrackRemoved was called. Use OnAudioTrackState with State [EpicRtcTrackState::Stopped]");
 	}
 
@@ -1649,7 +1625,7 @@ namespace UE::PixelStreaming2
 
 	void FEpicRtcStreamer::OnVideoTrackRemoved(EpicRtcVideoTrackInterface* VideoTrack)
 	{
-		// TODO (Eden.Harris) OnVideoTrackRemoved is not hooked up in EpicRtc. 
+		// TODO (Eden.Harris) OnVideoTrackRemoved is not hooked up in EpicRtc.
 		UE_LOGFMT(LogPixelStreaming2RTC, Warning, "FEpicRtcStreamer::OnVideoTrackRemoved was called. Use OnAudioTrackState with State [EpicRtcTrackState::Stopped]");
 	}
 
@@ -1703,10 +1679,13 @@ namespace UE::PixelStreaming2
 		}
 	}
 
-	void FEpicRtcStreamer::OnDataTrackRemoved(EpicRtcDataTrackInterface* DataTrack)
+	void FEpicRtcStreamer::OnVideoTrackEncodedFrame(EpicRtcVideoTrackInterface* VideoTrack, const EpicRtcEncodedVideoFrame& EncodedFrame)
 	{
-		// As long as DataTrack emits `stopped` state when it removed, this should be enough.
-		// because OnDataTrackState(DataTrack, EpicRtcTrackState::Stopped) already calls `OnDataChannelClosed`.
+	}
+
+	EpicRtcBool FEpicRtcStreamer::Enabled() const
+	{
+		return true;
 	}
 
 	void FEpicRtcStreamer::OnDataTrackState(EpicRtcDataTrackInterface* DataTrack, const EpicRtcTrackState State)
@@ -1807,6 +1786,10 @@ namespace UE::PixelStreaming2
 		}
 	}
 
+	void FEpicRtcStreamer::OnDataTrackError(EpicRtcDataTrackInterface* DataTrack, const EpicRtcErrorCode Error)
+	{
+	}
+
 	bool FEpicRtcStreamer::FindPlayerByAudioTrack(EpicRtcAudioTrackInterface* AudioTrack, FString& OutPlayerId)
 	{
 		OutPlayerId = "";
@@ -1888,11 +1871,6 @@ namespace UE::PixelStreaming2
 		return !OutPlayerId.IsEmpty();
 	}
 
-	FRTCStreamerFactory::FRTCStreamerFactory(TRefCountPtr<EpicRtcConferenceInterface> Conference)
-		: Conference(Conference)
-	{
-	}
-
 	FString FRTCStreamerFactory::GetStreamType()
 	{
 		return RTC_STREAM_TYPE;
@@ -1900,7 +1878,7 @@ namespace UE::PixelStreaming2
 
 	TSharedPtr<IPixelStreaming2Streamer> FRTCStreamerFactory::CreateNewStreamer(const FString& StreamerId)
 	{
-		TSharedPtr<IPixelStreaming2Streamer> NewStreamer = MakeShared<FEpicRtcStreamer>(StreamerId, Conference);
+		TSharedPtr<IPixelStreaming2Streamer> NewStreamer = MakeShared<FEpicRtcStreamer>(StreamerId);
 
 		// default to the scene viewport if we have a game engine
 		if (UGameEngine* GameEngine = Cast<UGameEngine>(GEngine))
diff --git a/Engine/Plugins/Media/PixelStreaming2/Source/PixelStreaming2RTC/Private/EpicRtcStreamer.h b/Engine/Plugins/Media/PixelStreaming2/Source/PixelStreaming2RTC/Private/EpicRtcStreamer.h
--- a/Engine/Plugins/Media/PixelStreaming2/Source/PixelStreaming2RTC/Private/EpicRtcStreamer.h
+++ b/Engine/Plugins/Media/PixelStreaming2/Source/PixelStreaming2RTC/Private/EpicRtcStreamer.h
@@ -11,8 +11,16 @@
 #include "PlayerContext.h"
 #include "FreezeFrame.h"
 #include "StreamerReconnectTimer.h"
-#include "EpicRtcManager.h"
-
+#include "EpicRtcAudioTrackObserver.h"
+#include "EpicRtcAudioTrackObserverFactory.h"
+#include "EpicRtcDataTrackObserver.h"
+#include "EpicRtcDataTrackObserverFactory.h"
+#include "EpicRtcVideoTrackObserver.h"
+#include "EpicRtcVideoTrackObserverFactory.h"
+#include "EpicRtcRoomObserver.h"
+#include "EpicRtcSessionObserver.h"
+
+#include "epic_rtc/core/conference.h"
 #include "epic_rtc/core/stats.h"
 
 class IPixelStreaming2RTCModule;
@@ -22,11 +30,17 @@ namespace UE::PixelStreaming2
 	static const FString INVALID_PLAYER_ID = TEXT("Invalid Player Id");
 	static const FString RTC_STREAM_TYPE = TEXT("DefaultRtc");
 
-	class FEpicRtcStreamer : public IPixelStreaming2Streamer, public TSharedFromThis<FEpicRtcStreamer>
+	class FEpicRtcStreamer :
+		public IPixelStreaming2Streamer,
+		public TSharedFromThis<FEpicRtcStreamer>,
+		public IPixelStreaming2SessionObserver,
+		public IPixelStreaming2RoomObserver,
+		public IPixelStreaming2AudioTrackObserver,
+		public IPixelStreaming2DataTrackObserver,
+		public IPixelStreaming2VideoTrackObserver
 	{
 	public:
-		FEpicRtcStreamer(const FString& StreamerId, TRefCountPtr<EpicRtcConferenceInterface> Conference);
-		static TSharedPtr<FEpicRtcStreamer> Create(const FString& StreamerId, TRefCountPtr<EpicRtcConferenceInterface> Conference);
+		FEpicRtcStreamer(const FString& StreamerId);
 		virtual ~FEpicRtcStreamer();
 
 		virtual void  Initialize() override;
@@ -136,49 +150,68 @@ namespace UE::PixelStreaming2
 
 		TSharedPtr<FStreamerReconnectTimer> ReconnectTimer;
 
-		TSharedPtr<FEpicRtcManager> EpicRtcManager;
-
 	public:
-		// Begin FEpicRtcManager Callbacks
-		void OnSessionStateUpdate(const EpicRtcSessionState StateUpdate);
-		void OnSessionErrorUpdate(const EpicRtcErrorCode ErrorUpdate);
-		void OnSessionRoomsAvailableUpdate(EpicRtcStringArrayInterface* RoomsList);
-
-		void OnRoomStateUpdate(const EpicRtcRoomState State);
-		void OnRoomJoinedUpdate(EpicRtcParticipantInterface* Participant);
-		void OnRoomLeftUpdate(const EpicRtcStringView ParticipantId);
-		void OnAudioTrackUpdate(EpicRtcParticipantInterface* Participant, EpicRtcAudioTrackInterface* AudioTrack);
-		void OnVideoTrackUpdate(EpicRtcParticipantInterface* Participant, EpicRtcVideoTrackInterface* VideoTrack);
-		void OnDataTrackUpdate(EpicRtcParticipantInterface* Participant, EpicRtcDataTrackInterface* DataTrack);
-		void OnLocalSdpUpdate(EpicRtcParticipantInterface* Participant, EpicRtcSdpInterface* Sdp);
-		void OnRemoteSdpUpdate(EpicRtcParticipantInterface* Participant, EpicRtcSdpInterface* Sdp);
-		void OnRoomErrorUpdate(const EpicRtcErrorCode Error);
-
-		void OnAudioTrackMuted(EpicRtcAudioTrackInterface* AudioTrack, EpicRtcBool bIsMuted);
-		void OnAudioTrackFrame(EpicRtcAudioTrackInterface* AudioTrack, const EpicRtcAudioFrame& Frame);
-		void OnAudioTrackRemoved(EpicRtcAudioTrackInterface* AudioTrack);
-		void OnAudioTrackState(EpicRtcAudioTrackInterface* AudioTrack, const EpicRtcTrackState State);
-
-		void OnVideoTrackMuted(EpicRtcVideoTrackInterface* VideoTrack, EpicRtcBool bIsMuted);
-		void OnVideoTrackFrame(EpicRtcVideoTrackInterface* VideoTrack, const EpicRtcVideoFrame& Frame);
-		void OnVideoTrackRemoved(EpicRtcVideoTrackInterface* VideoTrack);
-		void OnVideoTrackState(EpicRtcVideoTrackInterface* VideoTrack, const EpicRtcTrackState State);
-
-		void OnDataTrackRemoved(EpicRtcDataTrackInterface* DataTrack);
-		void OnDataTrackState(EpicRtcDataTrackInterface* DataTrack, const EpicRtcTrackState State);
-		void OnDataTrackMessage(EpicRtcDataTrackInterface* DataTrack);
-		// End FEpicRtcManager Callbacks
+		// Begin IPixelStreaming2SessionObserver
+		virtual void OnSessionStateUpdate(const EpicRtcSessionState StateUpdate) override;
+		virtual void OnSessionErrorUpdate(const EpicRtcErrorCode ErrorUpdate) override;
+		virtual void OnSessionRoomsAvailableUpdate(EpicRtcStringArrayInterface* RoomsList) override;
+		// End IPixelStreaming2SessionObserver
+
+		// Begin IPixelStreaming2RoomObserver
+		virtual void							   OnRoomStateUpdate(const EpicRtcRoomState State) override;
+		virtual void							   OnRoomJoinedUpdate(EpicRtcParticipantInterface* Participant) override;
+		virtual void							   OnRoomLeftUpdate(const EpicRtcStringView ParticipantId) override;
+		virtual void							   OnAudioTrackUpdate(EpicRtcParticipantInterface* Participant, EpicRtcAudioTrackInterface* AudioTrack) override;
+		virtual void							   OnVideoTrackUpdate(EpicRtcParticipantInterface* Participant, EpicRtcVideoTrackInterface* VideoTrack) override;
+		virtual void							   OnDataTrackUpdate(EpicRtcParticipantInterface* Participant, EpicRtcDataTrackInterface* DataTrack) override;
+		[[nodiscard]] virtual EpicRtcSdpInterface* OnLocalSdpUpdate(EpicRtcParticipantInterface* Participant, EpicRtcSdpInterface* Sdp) override;
+		[[nodiscard]] virtual EpicRtcSdpInterface* OnRemoteSdpUpdate(EpicRtcParticipantInterface* Participant, EpicRtcSdpInterface* Sdp) override;
+		virtual void							   OnRoomErrorUpdate(const EpicRtcErrorCode Error) override;
+		// End IPixelStreaming2RoomObserver
+
+		// Begin IPixelStreaming2AudioTrackObserver
+		virtual void OnAudioTrackMuted(EpicRtcAudioTrackInterface* AudioTrack, EpicRtcBool bIsMuted) override;
+		virtual void OnAudioTrackFrame(EpicRtcAudioTrackInterface* AudioTrack, const EpicRtcAudioFrame& Frame) override;
+		virtual void OnAudioTrackRemoved(EpicRtcAudioTrackInterface* AudioTrack) override;
+		virtual void OnAudioTrackState(EpicRtcAudioTrackInterface* AudioTrack, const EpicRtcTrackState State) override;
+		// End IPixelStreaming2AudioTrackObserver
+
+		// Begin IPixelStreaming2VideoTrackObserver
+		virtual void		OnVideoTrackMuted(EpicRtcVideoTrackInterface* VideoTrack, EpicRtcBool bIsMuted) override;
+		virtual void		OnVideoTrackFrame(EpicRtcVideoTrackInterface* VideoTrack, const EpicRtcVideoFrame& Frame) override;
+		virtual void		OnVideoTrackRemoved(EpicRtcVideoTrackInterface* VideoTrack) override;
+		virtual void		OnVideoTrackState(EpicRtcVideoTrackInterface* VideoTrack, const EpicRtcTrackState State) override;
+		virtual void		OnVideoTrackEncodedFrame(EpicRtcVideoTrackInterface*, const EpicRtcEncodedVideoFrame&) override;
+		virtual EpicRtcBool Enabled() const override;
+		// End IPixelStreaming2VideoTrackObserver
+
+		// Begin IPixelStreaming2DataTrackObserver
+		virtual void OnDataTrackState(EpicRtcDataTrackInterface* DataTrack, const EpicRtcTrackState State) override;
+		virtual void OnDataTrackMessage(EpicRtcDataTrackInterface* DataTrack) override;
+		virtual void OnDataTrackError(EpicRtcDataTrackInterface* DataTrack, const EpicRtcErrorCode Error) override;
+		// End IPixelStreaming2DataTrackObserver
+
+	private:
+		// Begin EpicRtc Classes
+		TRefCountPtr<EpicRtcSessionInterface>	 EpicRtcSession;
+		TRefCountPtr<EpicRtcRoomInterface>		 EpicRtcRoom;
+		// End EpicRtc Classes
+
+		// Begin EpicRtc Observers
+		TRefCountPtr<FEpicRtcSessionObserver>			SessionObserver;
+		TRefCountPtr<FEpicRtcRoomObserver>				RoomObserver;
+		TRefCountPtr<FEpicRtcAudioTrackObserverFactory> AudioTrackObserverFactory;
+		TRefCountPtr<FEpicRtcVideoTrackObserverFactory> VideoTrackObserverFactory;
+		TRefCountPtr<FEpicRtcDataTrackObserverFactory>	DataTrackObserverFactory;
+		// End EpicRtc Observers
 	};
 
 	class FRTCStreamerFactory : public IPixelStreaming2StreamerFactory
 	{
 	public:
-		FRTCStreamerFactory(TRefCountPtr<EpicRtcConferenceInterface> Conference);
+		FRTCStreamerFactory() = default;
 
 		virtual FString								 GetStreamType() override;
 		virtual TSharedPtr<IPixelStreaming2Streamer> CreateNewStreamer(const FString& StreamerId) override;
-
-	private:
-		TRefCountPtr<EpicRtcConferenceInterface> Conference;
 	};
 } // namespace UE::PixelStreaming2
diff --git a/Engine/Plugins/Media/PixelStreaming2/Source/PixelStreaming2RTC/Private/EpicRtcVideoTrackObserver.cpp b/Engine/Plugins/Media/PixelStreaming2/Source/PixelStreaming2RTC/Private/EpicRtcVideoTrackObserver.cpp
--- a/Engine/Plugins/Media/PixelStreaming2/Source/PixelStreaming2RTC/Private/EpicRtcVideoTrackObserver.cpp
+++ b/Engine/Plugins/Media/PixelStreaming2/Source/PixelStreaming2RTC/Private/EpicRtcVideoTrackObserver.cpp
@@ -2,54 +2,61 @@
 
 #include "EpicRtcVideoTrackObserver.h"
 
-#include "EpicRtcManager.h"
-
 namespace UE::PixelStreaming2
 {
 
-	FEpicRtcVideoTrackObserver::FEpicRtcVideoTrackObserver(TWeakPtr<FEpicRtcManager> Manager)
-		: Manager(Manager)
+	FEpicRtcVideoTrackObserver::FEpicRtcVideoTrackObserver(TObserverVariant<IPixelStreaming2VideoTrackObserver> UserObserver)
+		: UserObserver(UserObserver)
 	{
 	}
 
 	void FEpicRtcVideoTrackObserver::OnVideoTrackMuted(EpicRtcVideoTrackInterface* VideoTrack, EpicRtcBool bIsMuted)
 	{
-		if (!Manager.IsValid())
+		if (UserObserver)
 		{
-			return;
+			return UserObserver->OnVideoTrackMuted(VideoTrack, bIsMuted);
 		}
-
-		Manager.Pin()->OnVideoTrackMuted.Broadcast(VideoTrack, bIsMuted);
 	}
 
 	void FEpicRtcVideoTrackObserver::OnVideoTrackFrame(EpicRtcVideoTrackInterface* VideoTrack, const EpicRtcVideoFrame& Frame)
 	{
-		if (!Manager.IsValid())
+		if (UserObserver)
 		{
-			return;
+			return UserObserver->OnVideoTrackFrame(VideoTrack, Frame);
 		}
-
-		Manager.Pin()->OnVideoTrackFrame.Broadcast(VideoTrack, Frame);
 	}
 
 	void FEpicRtcVideoTrackObserver::OnVideoTrackRemoved(EpicRtcVideoTrackInterface* VideoTrack)
 	{
-		if (!Manager.IsValid())
+		if (UserObserver)
 		{
-			return;
+			return UserObserver->OnVideoTrackRemoved(VideoTrack);
 		}
-
-		Manager.Pin()->OnVideoTrackRemoved.Broadcast(VideoTrack);
 	}
 
 	void FEpicRtcVideoTrackObserver::OnVideoTrackState(EpicRtcVideoTrackInterface* VideoTrack, const EpicRtcTrackState State)
 	{
-		if (!Manager.IsValid())
+		if (UserObserver)
 		{
-			return;
+			return UserObserver->OnVideoTrackState(VideoTrack, State);
 		}
+	}
 
-		Manager.Pin()->OnVideoTrackState.Broadcast(VideoTrack, State);
+	void FEpicRtcVideoTrackObserver::OnVideoTrackEncodedFrame(EpicRtcVideoTrackInterface* VideoTrack, const EpicRtcEncodedVideoFrame& EncodedFrame)
+	{
+		if (UserObserver)
+		{
+			return UserObserver->OnVideoTrackEncodedFrame(VideoTrack, EncodedFrame);
+		}
+	}
+
+	EpicRtcBool FEpicRtcVideoTrackObserver::Enabled() const
+	{
+		if (UserObserver)
+		{
+			return UserObserver->Enabled();
+		}
+		return false;
 	}
 
 } // namespace UE::PixelStreaming2
\ No newline at end of file
diff --git a/Engine/Plugins/Media/PixelStreaming2/Source/PixelStreaming2RTC/Private/EpicRtcVideoTrackObserverFactory.cpp b/Engine/Plugins/Media/PixelStreaming2/Source/PixelStreaming2RTC/Private/EpicRtcVideoTrackObserverFactory.cpp
--- a/Engine/Plugins/Media/PixelStreaming2/Source/PixelStreaming2RTC/Private/EpicRtcVideoTrackObserverFactory.cpp
+++ b/Engine/Plugins/Media/PixelStreaming2/Source/PixelStreaming2RTC/Private/EpicRtcVideoTrackObserverFactory.cpp
@@ -2,20 +2,18 @@
 
 #include "EpicRtcVideoTrackObserverFactory.h"
 
-#include "EpicRtcManager.h"
 #include "EpicRtcVideoTrackObserver.h"
 
 namespace UE::PixelStreaming2
 {
-
-	FEpicRtcVideoTrackObserverFactory::FEpicRtcVideoTrackObserverFactory(TWeakPtr<FEpicRtcManager> Manager)
-		: Manager(Manager)
+	FEpicRtcVideoTrackObserverFactory::FEpicRtcVideoTrackObserverFactory(TObserverVariant<IPixelStreaming2VideoTrackObserver> UserObserver)
+		: UserObserver(UserObserver)
 	{
 	}
 
 	EpicRtcErrorCode FEpicRtcVideoTrackObserverFactory::CreateVideoTrackObserver(const EpicRtcStringView ParticipantId, const EpicRtcStringView VideoTrackId, EpicRtcVideoTrackObserverInterface** OutVideoTrackObserver)
 	{
-		EpicRtcVideoTrackObserverInterface* VideoTrackObserver = new FEpicRtcVideoTrackObserver(Manager);
+		EpicRtcVideoTrackObserverInterface* VideoTrackObserver = new FEpicRtcVideoTrackObserver(UserObserver);
 		// Because the ptr was created with new, we need to call AddRef ourself (ms spec compliant)
 		VideoTrackObserver->AddRef();
 
diff --git a/Engine/Plugins/Media/PixelStreaming2/Source/PixelStreaming2RTC/Private/PixelStreaming2RTCModule.cpp b/Engine/Plugins/Media/PixelStreaming2/Source/PixelStreaming2RTC/Private/PixelStreaming2RTCModule.cpp
--- a/Engine/Plugins/Media/PixelStreaming2/Source/PixelStreaming2RTC/Private/PixelStreaming2RTCModule.cpp
+++ b/Engine/Plugins/Media/PixelStreaming2/Source/PixelStreaming2RTC/Private/PixelStreaming2RTCModule.cpp
@@ -123,7 +123,7 @@ namespace UE::PixelStreaming2
 				return;
 			}
 
-			StreamerFactory.Reset(new FRTCStreamerFactory(EpicRtcConference));
+			StreamerFactory.Reset(new FRTCStreamerFactory());
 
 			// Ensure we have ImageWrapper loaded, used in Freezeframes
 			verify(FModuleManager::Get().LoadModule(FName("ImageWrapper")));
diff --git a/Engine/Plugins/Media/PixelStreaming2/Source/PixelStreaming2RTC/Private/PixelStreaming2RTCModule.h b/Engine/Plugins/Media/PixelStreaming2/Source/PixelStreaming2RTC/Private/PixelStreaming2RTCModule.h
--- a/Engine/Plugins/Media/PixelStreaming2/Source/PixelStreaming2RTC/Private/PixelStreaming2RTCModule.h
+++ b/Engine/Plugins/Media/PixelStreaming2/Source/PixelStreaming2RTC/Private/PixelStreaming2RTCModule.h
@@ -35,8 +35,8 @@ namespace UE::PixelStreaming2
 		// End IPixelStreaming2RTCModule
 
 		TSharedPtr<class FEpicRtcAudioCapturer>	 GetAudioCapturer();
-		TRefCountPtr<EpicRtcConferenceInterface> GetEpicRtcConference() { return EpicRtcConference; }
-		TRefCountPtr<FEpicRtcStatsCollector>	 GetStatsCollector() { return StatsCollector; }
+		TRefCountPtr<EpicRtcConferenceInterface>& GetEpicRtcConference() { return EpicRtcConference; }
+		TRefCountPtr<FEpicRtcStatsCollector>&	 GetStatsCollector() { return StatsCollector; }
 
 	private:
 		// Begin IModuleInterface
diff --git a/Engine/Plugins/Media/PixelStreaming2/Source/PixelStreaming2RTC/Private/Tests/EpicRtcRoomTest.cpp b/Engine/Plugins/Media/PixelStreaming2/Source/PixelStreaming2RTC/Private/Tests/EpicRtcRoomTest.cpp
--- a/Engine/Plugins/Media/PixelStreaming2/Source/PixelStreaming2RTC/Private/Tests/EpicRtcRoomTest.cpp
+++ b/Engine/Plugins/Media/PixelStreaming2/Source/PixelStreaming2RTC/Private/Tests/EpicRtcRoomTest.cpp
@@ -54,16 +54,16 @@ namespace UE::PixelStreaming2
 			MockWebsocketConnection->OnMessage().Broadcast(R"({"type" : "identify"})");
 		});
 
-		Manager->GetSessionObserver() = MakeRefCount<FEpicRtcSessionObserver>(Manager);
-		Manager->GetRoomObserver() = MakeRefCount<FEpicRtcRoomObserver>(Manager);
-		Manager->GetAudioTrackObserverFactory() = MakeRefCount<FEpicRtcAudioTrackObserverFactory>(Manager);
-		Manager->GetDataTrackObserverFactory() = MakeRefCount<FEpicRtcDataTrackObserverFactory>(Manager);
-		Manager->GetVideoTrackObserverFactory() = MakeRefCount<FEpicRtcVideoTrackObserverFactory>(Manager);
+		Manager->GetSessionObserver() = MakeRefCount<FEpicRtcSessionObserver>(TObserver<IPixelStreaming2SessionObserver>(Manager));
+		Manager->GetRoomObserver() = MakeRefCount<FEpicRtcRoomObserver>(TObserver<IPixelStreaming2RoomObserver>(Manager));
+		Manager->GetAudioTrackObserverFactory() = MakeRefCount<FEpicRtcAudioTrackObserverFactory>(TObserver<IPixelStreaming2AudioTrackObserver>(Manager));
+		Manager->GetDataTrackObserverFactory() = MakeRefCount<FEpicRtcDataTrackObserverFactory>(TObserver<IPixelStreaming2DataTrackObserver>(Manager));
+		Manager->GetVideoTrackObserverFactory() = MakeRefCount<FEpicRtcVideoTrackObserverFactory>(TObserver<IPixelStreaming2VideoTrackObserver>(Manager));
 
 		EpicRtcSessionState SessionState = EpicRtcSessionState::Disconnected;
 		EpicRtcRoomState	RoomState = EpicRtcRoomState::Failed;
 
-		Manager->OnSessionStateUpdate.AddLambda([this, RoomId, Manager, &SessionState](const EpicRtcSessionState State) {
+		Manager->OnSessionStateUpdateNative.AddLambda([this, RoomId, Manager, &SessionState](const EpicRtcSessionState State) {
 			switch (State)
 			{
 				case EpicRtcSessionState::Connected:
@@ -100,11 +100,11 @@ namespace UE::PixelStreaming2
 			SessionState = State;
 		});
 
-		Manager->OnSessionErrorUpdate.AddLambda([](EpicRtcErrorCode Error) {
+		Manager->OnSessionErrorUpdateNative.AddLambda([](EpicRtcErrorCode Error) {
 			UE_LOG(LogPixelStreaming2RTC, Error, TEXT("Manager->OnSessionErrorUpdate() {%d}"), Error);
 		});
 
-		Manager->OnRoomStateUpdate.AddLambda([this, &RoomState](const EpicRtcRoomState State) {
+		Manager->OnRoomStateUpdateNative.AddLambda([this, &RoomState](const EpicRtcRoomState State) {
 			UE_LOG(LogPixelStreaming2RTC, Log, TEXT("Manager->OnRoomStateUpdate() Old State (%s), New State (%s)"), *ToString(RoomState), *ToString(State));
 			RoomState = State;
 		});
diff --git a/Engine/Plugins/Media/PixelStreaming2/Source/PixelStreaming2RTC/Private/Tests/EpicRtcSessionTest.cpp b/Engine/Plugins/Media/PixelStreaming2/Source/PixelStreaming2RTC/Private/Tests/EpicRtcSessionTest.cpp
--- a/Engine/Plugins/Media/PixelStreaming2/Source/PixelStreaming2RTC/Private/Tests/EpicRtcSessionTest.cpp
+++ b/Engine/Plugins/Media/PixelStreaming2/Source/PixelStreaming2RTC/Private/Tests/EpicRtcSessionTest.cpp
@@ -48,7 +48,7 @@ namespace UE::PixelStreaming2
 		}
 
 		// Initialize session requirements
-		Manager->GetSessionObserver() = MakeRefCount<FEpicRtcSessionObserver>(Manager);
+		Manager->GetSessionObserver() = MakeRefCount<FEpicRtcSessionObserver>(TObserver<IPixelStreaming2SessionObserver>(Manager));
 
 		// Create and validate session
 		EpicRtcSessionConfig SessionConfig = {
diff --git a/Engine/Plugins/Media/PixelStreaming2/Source/PixelStreaming2RTC/Private/Tests/EpicRtcUtils.h b/Engine/Plugins/Media/PixelStreaming2/Source/PixelStreaming2RTC/Private/Tests/EpicRtcUtils.h
--- a/Engine/Plugins/Media/PixelStreaming2/Source/PixelStreaming2RTC/Private/Tests/EpicRtcUtils.h
+++ b/Engine/Plugins/Media/PixelStreaming2/Source/PixelStreaming2RTC/Private/Tests/EpicRtcUtils.h
@@ -2,66 +2,250 @@
 
 #pragma once
 
-#include "EpicRtcManager.h"
 #include "EpicRtcWebsocketFactory.h"
 #include "IWebSocket.h"
 #include "Logging.h"
 #include "Misc/AutomationTest.h"
 #include "UtilsString.h"
 
+#include "EpicRtcAudioTrackObserver.h"
+#include "EpicRtcAudioTrackObserverFactory.h"
+#include "EpicRtcDataTrackObserver.h"
+#include "EpicRtcDataTrackObserverFactory.h"
+#include "EpicRtcVideoTrackObserver.h"
+#include "EpicRtcVideoTrackObserverFactory.h"
+#include "EpicRtcRoomObserver.h"
+#include "EpicRtcSessionObserver.h"
+
 #include "epic_rtc/common/common.h"
 #include "epic_rtc/core/conference.h"
 #include "epic_rtc/core/platform.h"
 
 namespace UE::PixelStreaming2
 {
-	// A mock manager class for tests to receive callbacks from EpicRtc. Typically, the controlling class will inherit from FEpicRtcManager
+	// A mock manager class for tests to receive callbacks from EpicRtc. Typically, the controlling class will inherit from the relevant observer interfaces
 	// and implement the methods itself (see streamer.cpp). However, we can't force the tests to inherit the class, so instead we have the
-	// mock manager and the test bodies bind to the events they're interested in
-	class FMockManager : public FEpicRtcManager
+	// mock manager and the test bodies bind to the event delegates they're interested in
+	class FMockManager :
+		public TSharedFromThis<FMockManager>,
+		public IPixelStreaming2SessionObserver,
+		public IPixelStreaming2RoomObserver,
+		public IPixelStreaming2AudioTrackObserver,
+		public IPixelStreaming2DataTrackObserver,
+		public IPixelStreaming2VideoTrackObserver
 	{
 	public:
-		// Begin FEpicRtcAudioTrackObserver Callbacks
-		void OnAudioTrackMuted(EpicRtcAudioTrackInterface* AudioTrack, EpicRtcBool bIsMuted)
+		// Begin IPixelStreaming2SessionObserver
+		DECLARE_TS_MULTICAST_DELEGATE_OneParam(FOnSessionStateUpdate, const EpicRtcSessionState);
+		FOnSessionStateUpdate OnSessionStateUpdateNative;
+
+		DECLARE_TS_MULTICAST_DELEGATE_OneParam(FOnSessionErrorUpdate, const EpicRtcErrorCode);
+		FOnSessionErrorUpdate OnSessionErrorUpdateNative;
+
+		DECLARE_TS_MULTICAST_DELEGATE_OneParam(FOnSessionRoomsAvailableUpdate, EpicRtcStringArrayInterface*);
+		FOnSessionRoomsAvailableUpdate OnSessionRoomsAvailableUpdateNative;
+		// End IPixelStreaming2SessionObserver
+
+		// Begin IPixelStreaming2RoomObserver
+		DECLARE_TS_MULTICAST_DELEGATE_OneParam(FOnRoomStateUpdate, const EpicRtcRoomState);
+		FOnRoomStateUpdate OnRoomStateUpdateNative;
+
+		DECLARE_TS_MULTICAST_DELEGATE_OneParam(FOnRoomJoinedUpdate, EpicRtcParticipantInterface*);
+		FOnRoomJoinedUpdate OnRoomJoinedUpdateNative;
+
+		DECLARE_TS_MULTICAST_DELEGATE_OneParam(FOnRoomLeftUpdate, const EpicRtcStringView);
+		FOnRoomLeftUpdate OnRoomLeftUpdateNative;
+
+		DECLARE_TS_MULTICAST_DELEGATE_TwoParams(FOnAudioTrackUpdate, EpicRtcParticipantInterface*, EpicRtcAudioTrackInterface*);
+		FOnAudioTrackUpdate OnAudioTrackUpdateNative;
+
+		DECLARE_TS_MULTICAST_DELEGATE_TwoParams(FOnVideoTrackUpdate, EpicRtcParticipantInterface*, EpicRtcVideoTrackInterface*);
+		FOnVideoTrackUpdate OnVideoTrackUpdateNative;
+
+		DECLARE_TS_MULTICAST_DELEGATE_TwoParams(FOnDataTrackUpdate, EpicRtcParticipantInterface*, EpicRtcDataTrackInterface*);
+		FOnDataTrackUpdate OnDataTrackUpdateNative;
+
+		DECLARE_TS_MULTICAST_DELEGATE_TwoParams(FOnLocalSdpUpdate, EpicRtcParticipantInterface*, EpicRtcSdpInterface*);
+		FOnLocalSdpUpdate OnLocalSdpUpdateNative;
+
+		DECLARE_TS_MULTICAST_DELEGATE_TwoParams(FOnRemoteSdpUpdate, EpicRtcParticipantInterface*, EpicRtcSdpInterface*);
+		FOnRemoteSdpUpdate OnRemoteSdpUpdateNative;
+
+		DECLARE_TS_MULTICAST_DELEGATE_OneParam(FOnRoomErrorUpdate, const EpicRtcErrorCode);
+		FOnRoomErrorUpdate OnRoomErrorUpdateNative;
+		// End IPixelStreaming2RoomObserver
+
+		// Begin IPixelStreaming2AudioTrackObserver
+		DECLARE_TS_MULTICAST_DELEGATE_TwoParams(FOnAudioTrackMuted, EpicRtcAudioTrackInterface*, EpicRtcBool);
+		FOnAudioTrackMuted OnAudioTrackMutedNative;
+
+		DECLARE_TS_MULTICAST_DELEGATE_TwoParams(FOnAudioTrackFrame, EpicRtcAudioTrackInterface*, const EpicRtcAudioFrame&);
+		FOnAudioTrackFrame OnAudioTrackFrameNative;
+
+		DECLARE_TS_MULTICAST_DELEGATE_OneParam(FOnAudioTrackRemoved, EpicRtcAudioTrackInterface*);
+		FOnAudioTrackRemoved OnAudioTrackRemovedNative;
+
+		DECLARE_TS_MULTICAST_DELEGATE_TwoParams(FOnAudioTrackState, EpicRtcAudioTrackInterface*, const EpicRtcTrackState);
+		FOnAudioTrackState OnAudioTrackStateNative;
+		// End IPixelStreaming2AudioTrackObserver
+
+		// Begin IPixelStreaming2VideoTrackObserver
+		DECLARE_TS_MULTICAST_DELEGATE_TwoParams(FOnVideoTrackMuted, EpicRtcVideoTrackInterface*, EpicRtcBool);
+		FOnVideoTrackMuted OnVideoTrackMutedNative;
+
+		DECLARE_TS_MULTICAST_DELEGATE_TwoParams(FOnVideoTrackFrame, EpicRtcVideoTrackInterface*, const EpicRtcVideoFrame&);
+		FOnVideoTrackFrame OnVideoTrackFrameNative;
+
+		DECLARE_TS_MULTICAST_DELEGATE_OneParam(FOnVideoTrackRemoved, EpicRtcVideoTrackInterface*);
+		FOnVideoTrackRemoved OnVideoTrackRemovedNative;
+
+		DECLARE_TS_MULTICAST_DELEGATE_TwoParams(FOnVideoTrackState, EpicRtcVideoTrackInterface*, const EpicRtcTrackState);
+		FOnVideoTrackState OnVideoTrackStateNative;
+		// End IPixelStreaming2VideoTrackObserver
+
+		// Begin IPixelStreaming2DataTrackObserver
+		DECLARE_TS_MULTICAST_DELEGATE_TwoParams(FOnDataTrackState, EpicRtcDataTrackInterface*, const EpicRtcTrackState);
+		FOnDataTrackState OnDataTrackStateNative;
+
+		DECLARE_TS_MULTICAST_DELEGATE_OneParam(FOnDataTrackMessage, EpicRtcDataTrackInterface*);
+		FOnDataTrackMessage OnDataTrackMessageNative;
+
+		DECLARE_TS_MULTICAST_DELEGATE_TwoParams(FOnDataTrackError, EpicRtcDataTrackInterface*, const EpicRtcErrorCode);
+		FOnDataTrackError OnDataTrackErrorNative;
+		// End IPixelStreaming2DataTrackObserver
+
+	public:
+		// Begin IPixelStreaming2SessionObserver
+		virtual void OnSessionStateUpdate(const EpicRtcSessionState StateUpdate) override
+		{
+			OnSessionStateUpdateNative.Broadcast(StateUpdate);
+		}
+		virtual void OnSessionErrorUpdate(const EpicRtcErrorCode ErrorUpdate) override
+		{
+			OnSessionErrorUpdateNative.Broadcast(ErrorUpdate);
+		}
+		virtual void OnSessionRoomsAvailableUpdate(EpicRtcStringArrayInterface* RoomsList) override
 		{
+			OnSessionRoomsAvailableUpdateNative.Broadcast(RoomsList);
 		}
-		void OnAudioTrackFrame(EpicRtcAudioTrackInterface* AudioTrack, const EpicRtcAudioFrame& Frame)
+		// End IPixelStreaming2SessionObserver
+
+		// Begin IPixelStreaming2RoomObserver
+		virtual void OnRoomStateUpdate(const EpicRtcRoomState State) override
+		{
+			OnRoomStateUpdateNative.Broadcast(State);
+		}
+		virtual void OnRoomJoinedUpdate(EpicRtcParticipantInterface* Participant) override
+		{
+			OnRoomJoinedUpdateNative.Broadcast(Participant);
+		}
+		virtual void OnRoomLeftUpdate(const EpicRtcStringView ParticipantId) override
+		{
+			OnRoomLeftUpdateNative.Broadcast(ParticipantId);
+		}
+		virtual void OnAudioTrackUpdate(EpicRtcParticipantInterface* Participant, EpicRtcAudioTrackInterface* AudioTrack) override
 		{
+			OnAudioTrackUpdateNative.Broadcast(Participant, AudioTrack);
 		}
-		void OnAudioTrackRemoved(EpicRtcAudioTrackInterface* AudioTrack)
+		virtual void OnVideoTrackUpdate(EpicRtcParticipantInterface* Participant, EpicRtcVideoTrackInterface* VideoTrack) override
 		{
+			OnVideoTrackUpdateNative.Broadcast(Participant, VideoTrack);
 		}
-		void OnAudioTrackState(EpicRtcAudioTrackInterface* AudioTrack, const EpicRtcTrackState State)
+		virtual void OnDataTrackUpdate(EpicRtcParticipantInterface* Participant, EpicRtcDataTrackInterface* DataTrack) override
 		{
+			OnDataTrackUpdateNative.Broadcast(Participant, DataTrack);
 		}
-		// End FEpicRtcAudioTrackObserver Callbacks
+		[[nodiscard]] virtual EpicRtcSdpInterface* OnLocalSdpUpdate(EpicRtcParticipantInterface* Participant, EpicRtcSdpInterface* Sdp) override
+		{
+			OnLocalSdpUpdateNative.Broadcast(Participant, Sdp);
+			return nullptr;
+		}
+		[[nodiscard]] virtual EpicRtcSdpInterface* OnRemoteSdpUpdate(EpicRtcParticipantInterface* Participant, EpicRtcSdpInterface* Sdp) override
+		{
+			OnRemoteSdpUpdateNative.Broadcast(Participant, Sdp);
+			return nullptr;
+		}
+		virtual void OnRoomErrorUpdate(const EpicRtcErrorCode Error) override
+		{
+			OnRoomErrorUpdateNative.Broadcast(Error);
+		}
+		// End IPixelStreaming2RoomObserver
 
-		// Begin FEpicRtcVideoTrackObserver Callbacks
-		void OnVideoTrackMuted(EpicRtcVideoTrackInterface* VideoTrack, EpicRtcBool bIsMuted)
+		// Begin IPixelStreaming2AudioTrackObserver
+		virtual void OnAudioTrackMuted(EpicRtcAudioTrackInterface* AudioTrack, EpicRtcBool bIsMuted) override
 		{
+			OnAudioTrackMutedNative.Broadcast(AudioTrack, bIsMuted);
 		}
-		void OnVideoTrackFrame(EpicRtcVideoTrackInterface* VideoTrack, const EpicRtcVideoFrame& Frame)
+		virtual void OnAudioTrackFrame(EpicRtcAudioTrackInterface* AudioTrack, const EpicRtcAudioFrame& Frame) override
 		{
+			OnAudioTrackFrameNative.Broadcast(AudioTrack, Frame);
 		}
-		void OnVideoTrackRemoved(EpicRtcVideoTrackInterface* VideoTrack)
+		virtual void OnAudioTrackRemoved(EpicRtcAudioTrackInterface* AudioTrack) override
 		{
+			OnAudioTrackRemovedNative.Broadcast(AudioTrack);
 		}
-		void OnVideoTrackState(EpicRtcVideoTrackInterface* VideoTrack, const EpicRtcTrackState State)
+		virtual void OnAudioTrackState(EpicRtcAudioTrackInterface* AudioTrack, const EpicRtcTrackState State) override
 		{
+			OnAudioTrackStateNative.Broadcast(AudioTrack, State);
 		}
-		// End FEpicRtcVideoTrackObserver Callbacks
+		// End IPixelStreaming2AudioTrackObserver
 
-		// Begin FEpicRtcDataTrackObserver Callbacks
-		void OnDataTrackRemoved(EpicRtcDataTrackInterface* DataTrack)
+		// Begin IPixelStreaming2VideoTrackObserver
+		virtual void OnVideoTrackMuted(EpicRtcVideoTrackInterface* VideoTrack, EpicRtcBool bIsMuted) override
+		{
+			OnVideoTrackMutedNative.Broadcast(VideoTrack, bIsMuted);
+		}
+		virtual void OnVideoTrackFrame(EpicRtcVideoTrackInterface* VideoTrack, const EpicRtcVideoFrame& Frame) override
 		{
+			OnVideoTrackFrameNative.Broadcast(VideoTrack, Frame);
 		}
-		void OnDataTrackState(EpicRtcDataTrackInterface* DataTrack, const EpicRtcTrackState State)
+		virtual void OnVideoTrackRemoved(EpicRtcVideoTrackInterface* VideoTrack) override
 		{
+			OnVideoTrackRemovedNative.Broadcast(VideoTrack);
 		}
-		void OnDataTrackMessage(EpicRtcDataTrackInterface* DataTrack)
+		virtual void OnVideoTrackState(EpicRtcVideoTrackInterface* VideoTrack, const EpicRtcTrackState State) override
 		{
+			OnVideoTrackStateNative.Broadcast(VideoTrack, State);
 		}
-		// End FEpicRtcDataTrackObserver Callbacks
+		virtual void OnVideoTrackEncodedFrame(EpicRtcVideoTrackInterface* VideoTrack, const EpicRtcEncodedVideoFrame& EncodedFrame) override
+		{
+		}
+
+		virtual EpicRtcBool Enabled() const override
+		{
+			return true;
+		}
+		// End IPixelStreaming2VideoTrackObserver
+
+		// Begin IPixelStreaming2DataTrackObserver
+		virtual void OnDataTrackState(EpicRtcDataTrackInterface* DataTrack, const EpicRtcTrackState State) override
+		{
+			OnDataTrackStateNative.Broadcast(DataTrack, State);
+		}
+		virtual void OnDataTrackMessage(EpicRtcDataTrackInterface* DataTrack) override
+		{
+			OnDataTrackMessageNative.Broadcast(DataTrack);
+		}
+		virtual void OnDataTrackError(EpicRtcDataTrackInterface* DataTrack, const EpicRtcErrorCode Error) override
+		{
+			OnDataTrackErrorNative.Broadcast(DataTrack, Error);
+		}
+		// End IPixelStreaming2DataTrackObserver
+
+	private:
+		// Begin EpicRtc Classes
+		TRefCountPtr<EpicRtcConferenceInterface> EpicRtcConference;
+		TRefCountPtr<EpicRtcSessionInterface>	 EpicRtcSession;
+		TRefCountPtr<EpicRtcRoomInterface>		 EpicRtcRoom;
+		// End EpicRtc Classes
+
+		// Begin EpicRtc Observers
+		TRefCountPtr<FEpicRtcSessionObserver>			SessionObserver;
+		TRefCountPtr<FEpicRtcRoomObserver>				RoomObserver;
+		TRefCountPtr<FEpicRtcAudioTrackObserverFactory> AudioTrackObserverFactory;
+		TRefCountPtr<FEpicRtcVideoTrackObserverFactory> VideoTrackObserverFactory;
+		TRefCountPtr<FEpicRtcDataTrackObserverFactory>	DataTrackObserverFactory;
+		// End EpicRtc Observers
 
 	public:
 		TRefCountPtr<EpicRtcConferenceInterface>& GetEpicRtcConference()
diff --git a/Engine/Plugins/Media/PixelStreaming2/Source/PixelStreaming2RTC/Private/Tests/MediaDirectionTest.cpp b/Engine/Plugins/Media/PixelStreaming2/Source/PixelStreaming2RTC/Private/Tests/MediaDirectionTest.cpp
--- a/Engine/Plugins/Media/PixelStreaming2/Source/PixelStreaming2RTC/Private/Tests/MediaDirectionTest.cpp
+++ b/Engine/Plugins/Media/PixelStreaming2/Source/PixelStreaming2RTC/Private/Tests/MediaDirectionTest.cpp
@@ -58,38 +58,29 @@ namespace UE::PixelStreaming2
 
 		TSharedPtr<bool> bLocalTrack = MakeShared<bool>(false);
 		TSharedPtr<bool> bRemoteTrack = MakeShared<bool>(false);
+
+		ADD_LATENT_AUTOMATION_COMMAND(FExecuteLambda([Streamer]() { Streamer->StartStreaming(); }))
+		ADD_LATENT_AUTOMATION_COMMAND(FWaitAndCheckStreamerBool(TEXT("Check streaming started"), 5.0, Streamer, bStreamingStarted, true))
+		ADD_LATENT_AUTOMATION_COMMAND(FExecuteLambda([Player, PlayerPort]() { Player->Connect(PlayerPort); }))
+		ADD_LATENT_AUTOMATION_COMMAND(FSubscribePlayerAfterStreamerConnectedOrTimeout(5.0, Streamer, Player, StreamerName))
+		// Wait 5 seconds to ensure connection is fully established
+		ADD_LATENT_AUTOMATION_COMMAND(FWaitSeconds(5.0))
+
 		if (Media == EMediaType::Audio)
 		{
-			Player->EpicRtcManager->OnAudioTrackUpdate.AddLambda([bLocalTrack, bRemoteTrack](EpicRtcParticipantInterface* Particpant, EpicRtcAudioTrackInterface* AudioTrack) {
-				if (AudioTrack->IsRemote())
-				{
-					*bRemoteTrack.Get() = true;
-				}
-				else
-				{
-					*bLocalTrack.Get() = true;
-				}
-			});
+			ADD_LATENT_AUTOMATION_COMMAND(FExecuteLambda([Player, bLocalTrack, bRemoteTrack]() {
+				*(bLocalTrack.Get()) = Player->GetHasLocalAudioTrack();
+				*(bRemoteTrack.Get()) = Player->GetHasRemoteAudioTrack();
+			}))
 		}
 		else if (Media == EMediaType::Video)
 		{
-			Player->EpicRtcManager->OnVideoTrackUpdate.AddLambda([bLocalTrack, bRemoteTrack](EpicRtcParticipantInterface* Particpant, EpicRtcVideoTrackInterface* VideoTrack) {
-				if (VideoTrack->IsRemote())
-				{
-					*bRemoteTrack.Get() = true;
-				}
-				else
-				{
-					*bLocalTrack.Get() = true;
-				}
-			});
+			ADD_LATENT_AUTOMATION_COMMAND(FExecuteLambda([Player, bLocalTrack, bRemoteTrack]() {
+				*(bLocalTrack.Get()) = Player->GetHasLocalVideoTrack();
+				*(bRemoteTrack.Get()) = Player->GetHasRemoteVideoTrack();
+			}))
 		}
 
-		ADD_LATENT_AUTOMATION_COMMAND(FExecuteLambda([Streamer]() { Streamer->StartStreaming(); }))
-		ADD_LATENT_AUTOMATION_COMMAND(FWaitAndCheckStreamerBool(TEXT("Check streaming started"), 5.0, Streamer, bStreamingStarted, true))
-		ADD_LATENT_AUTOMATION_COMMAND(FExecuteLambda([Player, PlayerPort]() { Player->Connect(PlayerPort); }))
-		ADD_LATENT_AUTOMATION_COMMAND(FSubscribePlayerAfterStreamerConnectedOrTimeout(5.0, Streamer, Player, StreamerName))
-
 		if (Direction == EMediaDirection::SendOnly)
 		{
 			// Check that the player has only a remote track
diff --git a/Engine/Plugins/Media/PixelStreaming2/Source/PixelStreaming2RTC/Private/Tests/MockPlayer.cpp b/Engine/Plugins/Media/PixelStreaming2/Source/PixelStreaming2RTC/Private/Tests/MockPlayer.cpp
--- a/Engine/Plugins/Media/PixelStreaming2/Source/PixelStreaming2RTC/Private/Tests/MockPlayer.cpp
+++ b/Engine/Plugins/Media/PixelStreaming2/Source/PixelStreaming2RTC/Private/Tests/MockPlayer.cpp
@@ -30,48 +30,29 @@ namespace UE::PixelStreaming2
 		VideoBuffer.SafeRelease();
 	}
 
-	FMockPlayer::FMockPlayer(FMockPlayerConfig Config)
-		: AudioDirection(Config.AudioDirection)
-		, VideoDirection(Config.VideoDirection)
+	TSharedPtr<FMockPlayer> FMockPlayer::Create(FMockPlayerConfig Config)
 	{
-		PlayerName = FUtf8String::Printf("MockPlayer%d", PlayerId++);
+		TSharedPtr<FMockPlayer> Player = MakeShareable(new FMockPlayer(Config));
+
+		TWeakPtr<FMockPlayer> WeakPlayer = Player;
+
+		Player->SessionObserver = MakeRefCount<FEpicRtcSessionObserver>(TObserver<IPixelStreaming2SessionObserver>(WeakPlayer));
+		Player->RoomObserver = MakeRefCount<FEpicRtcRoomObserver>(TObserver<IPixelStreaming2RoomObserver>(WeakPlayer));
+
+		Player->AudioTrackObserverFactory = MakeRefCount<FEpicRtcAudioTrackObserverFactory>(TObserver<IPixelStreaming2AudioTrackObserver>(WeakPlayer));
+		Player->VideoTrackObserverFactory = MakeRefCount<FEpicRtcVideoTrackObserverFactory>(TObserver<IPixelStreaming2VideoTrackObserver>(WeakPlayer));
+		Player->DataTrackObserverFactory = MakeRefCount<FEpicRtcDataTrackObserverFactory>(TObserver<IPixelStreaming2DataTrackObserver>(WeakPlayer));
+
+		Player->EpicRtcVideoEncoderInitializers = { new FEpicRtcVideoEncoderInitializer() };
+		Player->EpicRtcVideoDecoderInitializers = { new FEpicRtcVideoDecoderInitializer() };
+
 		FUtf8String ConferenceId = "test_conference";
 
-		EpicRtcVideoEncoderInitializers = { new FEpicRtcVideoEncoderInitializer() };
-		EpicRtcVideoDecoderInitializers = { new FEpicRtcVideoDecoderInitializer() };
-
-		EpicRtcManager = MakeShared<FEpicRtcManager>();
-
-		ToStreamerProtocol = UE::PixelStreaming2Input::GetDefaultToStreamerProtocol();
-
-		EpicRtcManager->OnSessionRoomsAvailableUpdate.AddRaw(this, &FMockPlayer::OnSessionRoomsAvailableUpdate);
-		EpicRtcManager->OnSessionErrorUpdate.AddRaw(this, &FMockPlayer::OnSessionErrorUpdate);
-		EpicRtcManager->OnRoomStateUpdate.AddRaw(this, &FMockPlayer::OnRoomStateUpdate);
-		EpicRtcManager->OnRoomJoinedUpdate.AddRaw(this, &FMockPlayer::OnRoomJoinedUpdate);
-		EpicRtcManager->OnRoomLeftUpdate.AddRaw(this, &FMockPlayer::OnRoomLeftUpdate);
-		EpicRtcManager->OnRoomErrorUpdate.AddRaw(this, &FMockPlayer::OnRoomErrorUpdate);
-		EpicRtcManager->OnSessionStateUpdate.AddRaw(this, &FMockPlayer::OnSessionStateUpdate);
-		EpicRtcManager->OnVideoTrackUpdate.AddRaw(this, &FMockPlayer::OnVideoTrackUpdate);
-		EpicRtcManager->OnVideoTrackFrame.AddRaw(this, &FMockPlayer::OnVideoTrackFrame);
-		EpicRtcManager->OnVideoTrackMuted.AddRaw(this, &FMockPlayer::OnVideoTrackMuted);
-		EpicRtcManager->OnVideoTrackRemoved.AddRaw(this, &FMockPlayer::OnVideoTrackRemoved);
-		EpicRtcManager->OnVideoTrackState.AddRaw(this, &FMockPlayer::OnVideoTrackState);
-		EpicRtcManager->OnAudioTrackUpdate.AddRaw(this, &FMockPlayer::OnAudioTrackUpdate);
-		EpicRtcManager->OnAudioTrackFrame.AddRaw(this, &FMockPlayer::OnAudioTrackFrame);
-		EpicRtcManager->OnAudioTrackMuted.AddRaw(this, &FMockPlayer::OnAudioTrackMuted);
-		EpicRtcManager->OnAudioTrackRemoved.AddRaw(this, &FMockPlayer::OnAudioTrackRemoved);
-		EpicRtcManager->OnAudioTrackState.AddRaw(this, &FMockPlayer::OnAudioTrackState);
-
-		EpicRtcManager->OnDataTrackMessage.AddRaw(this, &FMockPlayer::OnDataTrackMessage);
-		EpicRtcManager->OnDataTrackRemoved.AddRaw(this, &FMockPlayer::OnDataTrackRemoved);
-		EpicRtcManager->OnDataTrackState.AddRaw(this, &FMockPlayer::OnDataTrackState);
-		EpicRtcManager->OnDataTrackUpdate.AddRaw(this, &FMockPlayer::OnDataTrackUpdate);
-
-		EpicRtcErrorCode Result = GetOrCreatePlatform({}, Platform.GetInitReference());
+		EpicRtcErrorCode Result = GetOrCreatePlatform({}, Player->Platform.GetInitReference());
 
 		TRefCountPtr<FEpicRtcWebsocketFactory> WebsocketFactory = MakeRefCount<FEpicRtcWebsocketFactory>(false);
 
-		Result = Platform->CreateConference(ToEpicRtcStringView(ConferenceId),
+		Result = Player->Platform->CreateConference(ToEpicRtcStringView(ConferenceId),
 			{ ._websocketFactory = WebsocketFactory.GetReference(),
 				._signallingType = EpicRtcSignallingType::PixelStreaming,
 				._signingPlugin = nullptr,
@@ -83,31 +64,33 @@ namespace UE::PixelStreaming2
 					._audioDecoderInitializers = {},
 					._enableBuiltInAudioCodecs = true,
 				},
-				._videoConfig = { ._videoEncoderInitializers = { ._ptr = const_cast<const EpicRtcVideoEncoderInitializerInterface**>(EpicRtcVideoEncoderInitializers.GetData()), ._size = (uint64_t)EpicRtcVideoEncoderInitializers.Num() }, ._videoDecoderInitializers = { ._ptr = const_cast<const EpicRtcVideoDecoderInitializerInterface**>(EpicRtcVideoDecoderInitializers.GetData()), ._size = (uint64_t)EpicRtcVideoDecoderInitializers.Num() }, ._enableBuiltInVideoCodecs = false },
+				._videoConfig = { ._videoEncoderInitializers = { ._ptr = const_cast<const EpicRtcVideoEncoderInitializerInterface**>(Player->EpicRtcVideoEncoderInitializers.GetData()), ._size = (uint64_t)Player->EpicRtcVideoEncoderInitializers.Num() }, ._videoDecoderInitializers = { ._ptr = const_cast<const EpicRtcVideoDecoderInitializerInterface**>(Player->EpicRtcVideoDecoderInitializers.GetData()), ._size = (uint64_t)Player->EpicRtcVideoDecoderInitializers.Num() }, ._enableBuiltInVideoCodecs = false },
 				._fieldTrials = { ._fieldTrials = EpicRtcStringView{ ._ptr = nullptr, ._length = 0 }, ._isGlobal = 0 } },
-			EpicRtcManager->EpicRtcConference.GetInitReference());
-
-		TickConferenceTask = FPixelStreamingTickableTask::Create<FEpicRtcTickConferenceTask>(EpicRtcManager->EpicRtcConference, TEXT("FMockPlayer TickConferenceTask"));
+			Player->EpicRtcConference.GetInitReference());
 
-		EpicRtcManager->SessionObserver = new FEpicRtcSessionObserver(EpicRtcManager);
-		EpicRtcManager->RoomObserver = new FEpicRtcRoomObserver(EpicRtcManager);
+		Player->TickConferenceTask = FPixelStreamingTickableTask::Create<FEpicRtcTickConferenceTask>(Player->EpicRtcConference, TEXT("FMockPlayer TickConferenceTask"));
 
-		EpicRtcManager->AudioTrackObserverFactory = new FEpicRtcAudioTrackObserverFactory(EpicRtcManager);
-		EpicRtcManager->VideoTrackObserverFactory = new FEpicRtcVideoTrackObserverFactory(EpicRtcManager);
-		EpicRtcManager->DataTrackObserverFactory = new FEpicRtcDataTrackObserverFactory(EpicRtcManager);
+		return Player;
+	}
 
-		VideoSink = MakeShared<FMockVideoSink>();
+	FMockPlayer::FMockPlayer(FMockPlayerConfig Config)
+		: VideoSink(MakeShared<FMockVideoSink>())
+		, ToStreamerProtocol(UE::PixelStreaming2Input::GetDefaultToStreamerProtocol())
+		, PlayerName(FUtf8String::Printf("MockPlayer%d", PlayerId++))
+		, AudioDirection(Config.AudioDirection)
+		, VideoDirection(Config.VideoDirection)
+	{
 	}
 
 	FMockPlayer::~FMockPlayer()
 	{
 		Disconnect(TEXT("Mock player being destroyed"));
 
-		if (EpicRtcManager->EpicRtcConference)
+		if (EpicRtcConference)
 		{
-			EpicRtcManager->EpicRtcConference->RemoveSession(ToEpicRtcStringView(PlayerName));
+			EpicRtcConference->RemoveSession(ToEpicRtcStringView(PlayerName));
 
-			Platform->ReleaseConference(EpicRtcManager->EpicRtcConference->GetId());
+			Platform->ReleaseConference(EpicRtcConference->GetId());
 		}
 	}
 
@@ -119,17 +102,17 @@ namespace UE::PixelStreaming2
 		EpicRtcSessionConfig SessionConfig = {
 			._id = ToEpicRtcStringView(PlayerName),
 			._url = ToEpicRtcStringView(ConnectionUrl),
-			._observer = EpicRtcManager->SessionObserver.GetReference()
+			._observer = SessionObserver.GetReference()
 		};
 
-		EpicRtcErrorCode Result = EpicRtcManager->EpicRtcConference->CreateSession(SessionConfig, EpicRtcManager->EpicRtcSession.GetInitReference());
+		EpicRtcErrorCode Result = EpicRtcConference->CreateSession(SessionConfig, EpicRtcSession.GetInitReference());
 		if (Result != EpicRtcErrorCode::Ok)
 		{
 			UE_LOG(LogPixelStreaming2RTC, Error, TEXT("FMockPlayer Failed to create EpicRtc session"));
 			return;
 		}
 
-		Result = EpicRtcManager->EpicRtcSession->Connect();
+		Result = EpicRtcSession->Connect();
 		if (Result != EpicRtcErrorCode::Ok)
 		{
 			UE_LOG(LogPixelStreaming2RTC, Error, TEXT("FMockPlayer Failed to connect EpicRtcSession"));
@@ -160,20 +143,20 @@ namespace UE::PixelStreaming2
 			._id = ToEpicRtcStringView(SubscribedStream),
 			._connectionConfig = ConnectionConfig,
 			._ticket = EpicRtcStringView{ ._ptr = nullptr, ._length = 0 },
-			._observer = EpicRtcManager->RoomObserver,
-			._audioTrackObserverFactory = EpicRtcManager->AudioTrackObserverFactory,
-			._dataTrackObserverFactory = EpicRtcManager->DataTrackObserverFactory,
-			._videoTrackObserverFactory = EpicRtcManager->VideoTrackObserverFactory
+			._observer = RoomObserver,
+			._audioTrackObserverFactory = AudioTrackObserverFactory,
+			._dataTrackObserverFactory = DataTrackObserverFactory,
+			._videoTrackObserverFactory = VideoTrackObserverFactory
 		};
 
-		EpicRtcErrorCode Result = EpicRtcManager->EpicRtcSession->CreateRoom(RoomConfig, EpicRtcManager->EpicRtcRoom.GetInitReference());
+		EpicRtcErrorCode Result = EpicRtcSession->CreateRoom(RoomConfig, EpicRtcRoom.GetInitReference());
 		if (Result != EpicRtcErrorCode::Ok)
 		{
 			UE_LOG(LogPixelStreaming2RTC, Error, TEXT("Failed to create EpicRtc room"));
 			return false;
 		}
 
-		EpicRtcManager->EpicRtcRoom->Join();
+		EpicRtcRoom->Join();
 
 		return true;
 	}
@@ -182,7 +165,16 @@ namespace UE::PixelStreaming2
 	{
 		FString ParticipantId{ (int32)Participant->GetId()._length, Participant->GetId()._ptr };
 		FString VideoTrackId{ (int32)VideoTrack->GetId()._length, VideoTrack->GetId()._ptr };
-		UE_LOG(LogPixelStreaming2RTC, VeryVerbose, TEXT("FMockPlayer::OnVideoTrackUpdate(Participant [%s], VideoTrack [%s])"), *ParticipantId, *VideoTrackId);
+		UE_LOGFMT(LogPixelStreaming2RTC, VeryVerbose, "FMockPlayer::OnVideoTrackUpdate(Participant [{0}], VideoTrack [{1}], Remote [{2}])", ParticipantId, VideoTrackId, static_cast<bool>(VideoTrack->IsRemote()));
+
+		if (VideoTrack->IsRemote())
+		{
+			bHasRemoteVideoTrack = true;
+		}
+		else
+		{
+			bHasLocalVideoTrack = true;
+		}
 	}
 
 	void FMockPlayer::OnVideoTrackFrame(EpicRtcVideoTrackInterface* VideoTrack, const EpicRtcVideoFrame& Frame)
@@ -204,6 +196,15 @@ namespace UE::PixelStreaming2
 	{
 	}
 
+	void FMockPlayer::OnVideoTrackEncodedFrame(EpicRtcVideoTrackInterface* VideoTrack, const EpicRtcEncodedVideoFrame& EncodedFrame)
+	{
+	}
+
+	EpicRtcBool FMockPlayer::Enabled() const
+	{
+		return true;
+	}
+
 	void FMockPlayer::OnSessionRoomsAvailableUpdate(EpicRtcStringArrayInterface* RoomsList)
 	{
 	}
@@ -319,8 +320,6 @@ namespace UE::PixelStreaming2
 		});
 	}
 
-	void FMockPlayer::OnDataTrackRemoved(EpicRtcDataTrackInterface*) {}
-
 	void FMockPlayer::OnDataTrackState(EpicRtcDataTrackInterface*, const EpicRtcTrackState) {}
 
 	void FMockPlayer::OnDataTrackUpdate(EpicRtcParticipantInterface*, EpicRtcDataTrackInterface* InDataTrack)
@@ -328,10 +327,36 @@ namespace UE::PixelStreaming2
 		DataTrack = FEpicRtcDataTrack::Create(InDataTrack, ToStreamerProtocol);
 	}
 
-	void FMockPlayer::OnAudioTrackUpdate(EpicRtcParticipantInterface* Particpant, EpicRtcAudioTrackInterface* AudioTrack)
+	[[nodiscard]] EpicRtcSdpInterface* FMockPlayer::OnLocalSdpUpdate(EpicRtcParticipantInterface* Participant, EpicRtcSdpInterface* Sdp)
+	{
+		return nullptr;
+	}
+
+	[[nodiscard]] EpicRtcSdpInterface* FMockPlayer::OnRemoteSdpUpdate(EpicRtcParticipantInterface* Participant, EpicRtcSdpInterface* Sdp)
+	{
+		return nullptr;
+	}
+
+	void FMockPlayer::OnDataTrackError(EpicRtcDataTrackInterface* InDataTrack, const EpicRtcErrorCode Error)
 	{
 	}
 
+	void FMockPlayer::OnAudioTrackUpdate(EpicRtcParticipantInterface* Participant, EpicRtcAudioTrackInterface* AudioTrack)
+	{
+		FString ParticipantId{ (int32)Participant->GetId()._length, Participant->GetId()._ptr };
+		FString AudioTrackId{ (int32)AudioTrack->GetId()._length, AudioTrack->GetId()._ptr };
+		UE_LOGFMT(LogPixelStreaming2RTC, VeryVerbose, "FMockPlayer::OnAudioTrackUpdate(Participant [{0}], AudioTrack [{1}], Remote [{2}])", ParticipantId, AudioTrackId, static_cast<bool>(AudioTrack->IsRemote()));
+
+		if (AudioTrack->IsRemote())
+		{
+			bHasRemoteAudioTrack = true;
+		}
+		else
+		{
+			bHasLocalAudioTrack = true;
+		}
+	}
+
 	void FMockPlayer::OnAudioTrackMuted(EpicRtcAudioTrackInterface* AudioTrack, EpicRtcBool bIsMuted)
 	{
 	}
@@ -350,18 +375,18 @@ namespace UE::PixelStreaming2
 
 	void FMockPlayer::Disconnect(const FString& Reason)
 	{
-		if (!EpicRtcManager->EpicRtcSession)
+		if (!EpicRtcSession)
 		{
 			return;
 		}
 
-		if (EpicRtcManager->EpicRtcRoom)
+		if (EpicRtcRoom)
 		{
-			EpicRtcManager->EpicRtcRoom->Leave();
-			EpicRtcManager->EpicRtcSession->RemoveRoom(ToEpicRtcStringView(SubscribedStream));
+			EpicRtcRoom->Leave();
+			EpicRtcSession->RemoveRoom(ToEpicRtcStringView(SubscribedStream));
 		}
 
-		EpicRtcErrorCode Result = EpicRtcManager->EpicRtcSession->Disconnect(ToEpicRtcStringView(*Reason));
+		EpicRtcErrorCode Result = EpicRtcSession->Disconnect(ToEpicRtcStringView(*Reason));
 		if (Result != EpicRtcErrorCode::Ok)
 		{
 			UE_LOG(LogPixelStreaming2RTC, Error, TEXT("Failed to disconnect EpicRtcSession"));
diff --git a/Engine/Plugins/Media/PixelStreaming2/Source/PixelStreaming2RTC/Private/Tests/MockPlayer.h b/Engine/Plugins/Media/PixelStreaming2/Source/PixelStreaming2RTC/Private/Tests/MockPlayer.h
--- a/Engine/Plugins/Media/PixelStreaming2/Source/PixelStreaming2RTC/Private/Tests/MockPlayer.h
+++ b/Engine/Plugins/Media/PixelStreaming2/Source/PixelStreaming2RTC/Private/Tests/MockPlayer.h
@@ -4,11 +4,19 @@
 
 #include "EpicRtcDataTrack.h"
 #include "HAL/ThreadSafeBool.h"
-#include "EpicRtcManager.h"
 #include "EpicRtcConferenceUtils.h"
 #include "IPixelStreaming2DataProtocol.h"
 #include "Containers/Utf8String.h"
 
+#include "EpicRtcAudioTrackObserver.h"
+#include "EpicRtcAudioTrackObserverFactory.h"
+#include "EpicRtcDataTrackObserver.h"
+#include "EpicRtcDataTrackObserverFactory.h"
+#include "EpicRtcVideoTrackObserver.h"
+#include "EpicRtcVideoTrackObserverFactory.h"
+#include "EpicRtcRoomObserver.h"
+#include "EpicRtcSessionObserver.h"
+
 #if WITH_DEV_AUTOMATION_TESTS
 
 class EpicRtcPlatformInterface;
@@ -53,38 +61,22 @@ namespace UE::PixelStreaming2
 		EMediaDirection VideoDirection;
 	};
 
-	class FMockPlayer : public TSharedFromThis<FMockPlayer>
+	class FMockPlayer :
+		public TSharedFromThis<FMockPlayer>,
+		public IPixelStreaming2SessionObserver,
+		public IPixelStreaming2RoomObserver,
+		public IPixelStreaming2AudioTrackObserver,
+		public IPixelStreaming2DataTrackObserver,
+		public IPixelStreaming2VideoTrackObserver
 	{
 	public:
-		FMockPlayer(FMockPlayerConfig Config = { .AudioDirection = EMediaDirection::RecvOnly, .VideoDirection = EMediaDirection::RecvOnly });
+		static TSharedPtr<FMockPlayer> Create(FMockPlayerConfig Config = { .AudioDirection = EMediaDirection::RecvOnly, .VideoDirection = EMediaDirection::RecvOnly });
 		virtual ~FMockPlayer();
 
 		void Connect(int StreamerPort);
 		void Disconnect(const FString& Reason);
 		bool Subscribe(const FString& StreamerId);
 
-		void OnVideoTrackUpdate(EpicRtcParticipantInterface* Participant, EpicRtcVideoTrackInterface* VideoTrack);
-		void OnVideoTrackFrame(EpicRtcVideoTrackInterface* VideoTrack, const EpicRtcVideoFrame& Frame);
-		void OnVideoTrackMuted(EpicRtcVideoTrackInterface* VideoTrack, EpicRtcBool bIsMuted);
-		void OnVideoTrackRemoved(EpicRtcVideoTrackInterface* VideoTrack);
-		void OnVideoTrackState(EpicRtcVideoTrackInterface* VideoTrack, const EpicRtcTrackState State);
-		void OnSessionRoomsAvailableUpdate(EpicRtcStringArrayInterface* RoomsList);
-		void OnSessionErrorUpdate(const EpicRtcErrorCode ErrorUpdate);
-		void OnRoomStateUpdate(const EpicRtcRoomState State);
-		void OnRoomJoinedUpdate(EpicRtcParticipantInterface* Participant);
-		void OnRoomLeftUpdate(const EpicRtcStringView ParticipantId);
-		void OnRoomErrorUpdate(const EpicRtcErrorCode Error);
-		void OnSessionStateUpdate(const EpicRtcSessionState StateUpdate);
-		void OnDataTrackMessage(EpicRtcDataTrackInterface*);
-		void OnDataTrackRemoved(EpicRtcDataTrackInterface*);
-		void OnDataTrackState(EpicRtcDataTrackInterface*, const EpicRtcTrackState);
-		void OnDataTrackUpdate(EpicRtcParticipantInterface*, EpicRtcDataTrackInterface*);
-		void OnAudioTrackUpdate(EpicRtcParticipantInterface* Particpant, EpicRtcAudioTrackInterface* AudioTrack);
-		void OnAudioTrackMuted(EpicRtcAudioTrackInterface* AudioTrack, EpicRtcBool bIsMuted);
-		void OnAudioTrackFrame(EpicRtcAudioTrackInterface* AudioTrack, const EpicRtcAudioFrame& Frame);
-		void OnAudioTrackRemoved(EpicRtcAudioTrackInterface* AudioTrack);
-		void OnAudioTrackState(EpicRtcAudioTrackInterface* AudioTrack, const EpicRtcTrackState);
-
 		template <typename... Args>
 		bool SendMessage(FString MessageType, Args... VarArgs)
 		{
@@ -101,12 +93,73 @@ namespace UE::PixelStreaming2
 		TSharedPtr<FMockVideoSink>				 GetVideoSink() { return VideoSink; };
 		TSharedPtr<IPixelStreaming2DataProtocol> GetToStreamerProtocol() { return ToStreamerProtocol; }
 
+		bool GetHasLocalAudioTrack() const { return bHasLocalAudioTrack; };
+		bool GetHasRemoteAudioTrack() const { return bHasRemoteAudioTrack; };
+		bool GetHasLocalVideoTrack() const { return bHasLocalVideoTrack; };
+		bool GetHasRemoteVideoTrack() const { return bHasRemoteVideoTrack; };
+
 		DECLARE_MULTICAST_DELEGATE_OneParam(FOnMessageReceived, const TArray<uint8>&);
 		FOnMessageReceived OnMessageReceived;
 
-		TSharedPtr<FEpicRtcManager> EpicRtcManager;
+	public:
+		// Begin IPixelStreaming2SessionObserver
+		virtual void OnSessionStateUpdate(const EpicRtcSessionState StateUpdate) override;
+		virtual void OnSessionErrorUpdate(const EpicRtcErrorCode ErrorUpdate) override;
+		virtual void OnSessionRoomsAvailableUpdate(EpicRtcStringArrayInterface* RoomsList) override;
+		// End IPixelStreaming2SessionObserver
+
+		// Begin IPixelStreaming2RoomObserver
+		virtual void							   OnRoomStateUpdate(const EpicRtcRoomState State) override;
+		virtual void							   OnRoomJoinedUpdate(EpicRtcParticipantInterface* Participant) override;
+		virtual void							   OnRoomLeftUpdate(const EpicRtcStringView ParticipantId) override;
+		virtual void							   OnAudioTrackUpdate(EpicRtcParticipantInterface* Participant, EpicRtcAudioTrackInterface* AudioTrack) override;
+		virtual void							   OnVideoTrackUpdate(EpicRtcParticipantInterface* Participant, EpicRtcVideoTrackInterface* VideoTrack) override;
+		virtual void							   OnDataTrackUpdate(EpicRtcParticipantInterface* Participant, EpicRtcDataTrackInterface* DataTrack) override;
+		[[nodiscard]] virtual EpicRtcSdpInterface* OnLocalSdpUpdate(EpicRtcParticipantInterface* Participant, EpicRtcSdpInterface* Sdp) override;
+		[[nodiscard]] virtual EpicRtcSdpInterface* OnRemoteSdpUpdate(EpicRtcParticipantInterface* Participant, EpicRtcSdpInterface* Sdp) override;
+		virtual void							   OnRoomErrorUpdate(const EpicRtcErrorCode Error) override;
+		// End IPixelStreaming2RoomObserver
+
+		// Begin IPixelStreaming2AudioTrackObserver
+		virtual void OnAudioTrackMuted(EpicRtcAudioTrackInterface* AudioTrack, EpicRtcBool bIsMuted) override;
+		virtual void OnAudioTrackFrame(EpicRtcAudioTrackInterface* AudioTrack, const EpicRtcAudioFrame& Frame) override;
+		virtual void OnAudioTrackRemoved(EpicRtcAudioTrackInterface* AudioTrack) override;
+		virtual void OnAudioTrackState(EpicRtcAudioTrackInterface* AudioTrack, const EpicRtcTrackState State) override;
+		// End IPixelStreaming2AudioTrackObserver
+
+		// Begin IPixelStreaming2VideoTrackObserver
+		virtual void		OnVideoTrackMuted(EpicRtcVideoTrackInterface* VideoTrack, EpicRtcBool bIsMuted) override;
+		virtual void		OnVideoTrackFrame(EpicRtcVideoTrackInterface* VideoTrack, const EpicRtcVideoFrame& Frame) override;
+		virtual void		OnVideoTrackRemoved(EpicRtcVideoTrackInterface* VideoTrack) override;
+		virtual void		OnVideoTrackState(EpicRtcVideoTrackInterface* VideoTrack, const EpicRtcTrackState State) override;
+		virtual void		OnVideoTrackEncodedFrame(EpicRtcVideoTrackInterface* VideoTrack, const EpicRtcEncodedVideoFrame& EncodedFrame) override;
+		virtual EpicRtcBool Enabled() const override;
+		// End IPixelStreaming2VideoTrackObserver
+
+		// Begin IPixelStreaming2DataTrackObserver
+		virtual void OnDataTrackState(EpicRtcDataTrackInterface* DataTrack, const EpicRtcTrackState State) override;
+		virtual void OnDataTrackMessage(EpicRtcDataTrackInterface* DataTrack) override;
+		virtual void OnDataTrackError(EpicRtcDataTrackInterface* DataTrack, const EpicRtcErrorCode Error) override;
+		// End IPixelStreaming2DataTrackObserver
+
+	private:
+		// Begin EpicRtc Classes
+		TRefCountPtr<EpicRtcConferenceInterface> EpicRtcConference;
+		TRefCountPtr<EpicRtcSessionInterface>	 EpicRtcSession;
+		TRefCountPtr<EpicRtcRoomInterface>		 EpicRtcRoom;
+		// End EpicRtc Classes
+
+		// Begin EpicRtc Observers
+		TRefCountPtr<FEpicRtcSessionObserver>			SessionObserver;
+		TRefCountPtr<FEpicRtcRoomObserver>				RoomObserver;
+		TRefCountPtr<FEpicRtcAudioTrackObserverFactory> AudioTrackObserverFactory;
+		TRefCountPtr<FEpicRtcVideoTrackObserverFactory> VideoTrackObserverFactory;
+		TRefCountPtr<FEpicRtcDataTrackObserverFactory>	DataTrackObserverFactory;
+		// End EpicRtc Observers
 
 	private:
+		FMockPlayer(FMockPlayerConfig Config);
+
 		TSharedPtr<FMockVideoSink>				   VideoSink;
 		TSharedPtr<FEpicRtcDataTrack>			   DataTrack;
 		TRefCountPtr<EpicRtcPlatformInterface>	   Platform;
@@ -124,6 +177,11 @@ namespace UE::PixelStreaming2
 
 		EMediaDirection AudioDirection;
 		EMediaDirection VideoDirection;
+
+		bool bHasLocalAudioTrack = false;
+		bool bHasRemoteAudioTrack = false;
+		bool bHasLocalVideoTrack = false;
+		bool bHasRemoteVideoTrack = false;
 	};
 } // namespace UE::PixelStreaming2
 
diff --git a/Engine/Plugins/Media/PixelStreaming2/Source/PixelStreaming2RTC/Private/Tests/TestUtils.cpp b/Engine/Plugins/Media/PixelStreaming2/Source/PixelStreaming2RTC/Private/Tests/TestUtils.cpp
--- a/Engine/Plugins/Media/PixelStreaming2/Source/PixelStreaming2RTC/Private/Tests/TestUtils.cpp
+++ b/Engine/Plugins/Media/PixelStreaming2/Source/PixelStreaming2RTC/Private/Tests/TestUtils.cpp
@@ -18,7 +18,7 @@
 #if WITH_DEV_AUTOMATION_TESTS
 
 namespace UE::PixelStreaming2
-{	
+{
 	int32 TestUtils::NextStreamerPort()
 	{
 		// Start of IANA un-registerable ports (49152 - 65535)
@@ -363,7 +363,7 @@ namespace UE::PixelStreaming2
 
 	TSharedPtr<FMockPlayer> CreatePlayer(FMockPlayerConfig Config)
 	{
-		TSharedPtr<FMockPlayer> OutPlayer = MakeShared<FMockPlayer>(Config);
+		TSharedPtr<FMockPlayer> OutPlayer = FMockPlayer::Create(Config);
 
 		return OutPlayer;
 	}
