[PS2] Added PS2 Delegates tests.

diff --git a/Engine/Plugins/Media/PixelStreaming2/Source/PixelStreaming2RTC/Private/Tests/DelegateTest.cpp b/Engine/Plugins/Media/PixelStreaming2/Source/PixelStreaming2RTC/Private/Tests/DelegateTest.cpp
new file mode 100644
--- /dev/null
+++ b/Engine/Plugins/Media/PixelStreaming2/Source/PixelStreaming2RTC/Private/Tests/DelegateTest.cpp
@@ -0,0 +1,474 @@
+// Copyright Epic Games, Inc. All Rights Reserved.
+
+#include "DelegateTest.h"
+
+#if WITH_TESTS
+
+#include "CodecUtils.h"
+#include "HAL/PlatformProcess.h"
+#include "Logging.h"
+#include "Misc/AutomationTest.h"
+#include "PixelStreaming2Delegates.h"
+#include "PixelStreaming2PluginSettings.h"
+#include "TestUtils.h"
+
+void UPixelStreaming2DynamicDelegateTest::OnConnectedToSignallingServer(FString StreamerId)
+{
+	DynamicDelegateCalled(TEXT("OnConnectedToSignallingServer"), StreamerId);
+}
+
+void UPixelStreaming2DynamicDelegateTest::OnDisconnectedFromSignallingServer(FString StreamerId)
+{
+	DynamicDelegateCalled(TEXT("OnDisconnectedFromSignallingServer"), StreamerId);
+}
+
+void UPixelStreaming2DynamicDelegateTest::OnNewConnection(FString StreamerId, FString PlayerId)
+{
+	DynamicDelegateCalled(TEXT("OnNewConnection"), StreamerId, PlayerId);
+}
+
+void UPixelStreaming2DynamicDelegateTest::OnClosedConnection(FString StreamerId, FString PlayerId)
+{
+	DynamicDelegateCalled(TEXT("OnClosedConnection"), StreamerId, PlayerId);
+}
+
+void UPixelStreaming2DynamicDelegateTest::OnAllConnectionsClosed(FString StreamerId)
+{
+	DynamicDelegateCalled(TEXT("OnAllConnectionsClosed"), StreamerId);
+}
+
+void UPixelStreaming2DynamicDelegateTest::OnDataTrackOpen(FString StreamerId, FString PlayerId)
+{
+	DynamicDelegateCalled(TEXT("OnDataTrackOpen"), StreamerId, PlayerId);
+}
+
+void UPixelStreaming2DynamicDelegateTest::OnDataTrackClosed(FString StreamerId, FString PlayerId)
+{
+	DynamicDelegateCalled(TEXT("OnDataTrackClosed"), StreamerId, PlayerId);
+}
+
+void UPixelStreaming2DynamicDelegateTest::OnStatChanged(FString PlayerId, FName StatName, float StatValue)
+{
+	DynamicDelegateCalled(TEXT("OnStatChanged"), PlayerId, StatName, StatValue);
+}
+
+void UPixelStreaming2DynamicDelegateTest::OnFallbackToSoftwareEncoding()
+{
+	DynamicDelegateCalled(TEXT("OnFallbackToSoftwareEncoding"));
+}
+
+// Macros used because passing the callback into a function results in a runtime check hit 
+// because UE checks the variable name of UFUNCTIONs
+#define BIND_DELEGATE(Delegate, Callback, Name, ...) 	\
+ 	Delegate.AddDynamic(this, Callback); 				\
+ 	BindDelegate<__VA_ARGS__>(Name)
+
+bool UPixelStreaming2DynamicDelegateTest::Init(UE::PixelStreaming2::DelegateTestConfig Config, FString StreamerName)
+{
+	UPixelStreaming2Delegates* Delegates = UPixelStreaming2Delegates::Get();
+	if (!Delegates)
+	{
+		UE_LOGFMT(LogPixelStreaming2RTC, Error, "Delegates are null.");
+		return false;
+	}
+	else
+	{
+		const bool bIsRemote = false;
+		
+		BIND_DELEGATE(Delegates->OnConnectedToSignallingServer, &UPixelStreaming2DynamicDelegateTest::OnConnectedToSignallingServer, TEXT("OnConnectedToSignallingServer"), FString)->Times(UE::PixelStreaming2::Exactly(1)).With(StreamerName);
+		BIND_DELEGATE(Delegates->OnDisconnectedFromSignallingServer, &UPixelStreaming2DynamicDelegateTest::OnDisconnectedFromSignallingServer, TEXT("OnDisconnectedFromSignallingServer"), FString)->Times(UE::PixelStreaming2::Exactly(1)).With(StreamerName);
+		BIND_DELEGATE(Delegates->OnNewConnection, &UPixelStreaming2DynamicDelegateTest::OnNewConnection, TEXT("OnNewConnection"), FString, TOptional<FString>)->Times(UE::PixelStreaming2::Exactly(Config.NumPlayers)).With(StreamerName, UE::PixelStreaming2::Any<FString>());
+		// TODO (Eden.Harris) This currently only fires once but should fire NumPlayers times.
+		BIND_DELEGATE(Delegates->OnClosedConnection, &UPixelStreaming2DynamicDelegateTest::OnClosedConnection, TEXT("OnClosedConnection"), FString, TOptional<FString>)->Times(UE::PixelStreaming2::Exactly(1)).With(StreamerName, UE::PixelStreaming2::Any<FString>());
+		BIND_DELEGATE(Delegates->OnAllConnectionsClosed, &UPixelStreaming2DynamicDelegateTest::OnAllConnectionsClosed, TEXT("OnAllConnectionsClosed"), FString)->Times(UE::PixelStreaming2::Exactly(1)).With(StreamerName);
+		BIND_DELEGATE(Delegates->OnDataTrackOpen, &UPixelStreaming2DynamicDelegateTest::OnDataTrackOpen, TEXT("OnDataTrackOpen"), FString, TOptional<FString>)->Times(UE::PixelStreaming2::Exactly(Config.NumPlayers)).With(StreamerName, UE::PixelStreaming2::Any<FString>());
+		// TODO (Eden.Harris) DataTrack closed does not currently fire
+		BIND_DELEGATE(Delegates->OnDataTrackClosed, &UPixelStreaming2DynamicDelegateTest::OnDataTrackClosed, TEXT("OnDataTrackClosed"), FString, TOptional<FString>)->Times(UE::PixelStreaming2::AtLeast(0)).With(StreamerName, UE::PixelStreaming2::Any<FString>());
+		BIND_DELEGATE(Delegates->OnStatChanged, &UPixelStreaming2DynamicDelegateTest::OnStatChanged, TEXT("OnStatChanged"), TOptional<FString>, TOptional<FName>, TOptional<float>)->Times(UE::PixelStreaming2::AtLeast(1)).With(UE::PixelStreaming2::Any<FString>(), UE::PixelStreaming2::Any<FName>(), UE::PixelStreaming2::Any<float>());
+		BIND_DELEGATE(Delegates->OnFallbackToSoftwareEncoding, &UPixelStreaming2DynamicDelegateTest::OnFallbackToSoftwareEncoding, TEXT("OnFallbackToSoftwareEncoding"))->Times(UE::PixelStreaming2::Exactly(Config.SoftwareEncodingCount));
+	}
+	return true;
+}
+
+void UPixelStreaming2DynamicDelegateTest::Destroy()
+{
+	UPixelStreaming2Delegates* Delegates = UPixelStreaming2Delegates::Get();
+	if (!Delegates)
+	{
+		UE_LOGFMT(LogPixelStreaming2RTC, Error, "Delegates are null.");
+	}
+	else
+	{
+		for (const auto& DelegateTest : DelegatesMap)
+		{
+			if (const auto& Value = DelegateTest.Value; !Value->bWasCalledExpectedTimes(true))
+			{
+				UE_LOGFMT(LogPixelStreaming2RTC, Error, "{0} was called {1} times.", Value->Name, Value->CallCount);
+			}
+		}
+		
+		Delegates->OnConnectedToSignallingServer.RemoveAll(this);
+		Delegates->OnDisconnectedFromSignallingServer.RemoveAll(this);
+		Delegates->OnNewConnection.RemoveAll(this);
+		Delegates->OnClosedConnection.RemoveAll(this);
+		Delegates->OnAllConnectionsClosed.RemoveAll(this);
+		Delegates->OnDataTrackOpen.RemoveAll(this);
+		Delegates->OnDataTrackClosed.RemoveAll(this);
+		Delegates->OnStatChanged.RemoveAll(this);
+		Delegates->OnFallbackToSoftwareEncoding.RemoveAll(this);
+	}
+	DelegatesMap.Empty();
+}
+
+namespace UE::PixelStreaming2
+{
+	FCardinality AnyNumber()
+	{
+		return FCardinality();
+	}
+
+	FCardinality AtLeast(int Min)
+	{
+		return FCardinality(Min, std::numeric_limits<int>::max());
+	}
+
+	FCardinality AtMost(int Max)
+	{
+		return FCardinality(std::numeric_limits<int>::min(), Max);
+	}
+
+	FCardinality Between(int Min, int Max)
+	{
+		return FCardinality(Min, Max);
+	}
+
+	FCardinality Exactly(int ExactValue)
+	{
+		return FCardinality(ExactValue, ExactValue);
+	}
+	
+	bool FDelegateTestBase::CheckCalled(bool bPrintErrors) const
+	{
+		for (const auto& DelegateTest : DelegatesMap)
+		{
+			if (!DelegateTest.Value->bWasCalledExpectedTimes(bPrintErrors))
+			{
+				return false;
+			}
+		}
+
+		return true;
+	}
+	// Used to hold onto the lifetime of the Dynamic delegate
+	class FDynamicDelegateLifetime
+	{
+	public:
+		FDynamicDelegateLifetime() = default;
+		~FDynamicDelegateLifetime() 
+		{
+			DelegateTest->Destroy();
+		}
+
+		bool Init(DelegateTestConfig Config, FString StreamerName)
+		{
+			DelegateTest = TStrongObjectPtr<UPixelStreaming2DynamicDelegateTest>(NewObject<UPixelStreaming2DynamicDelegateTest>());
+			if (!DelegateTest->Init(Config, StreamerName))
+			{
+				UE_LOGFMT(LogPixelStreaming2RTC, Error, "Unable to create FDelegatesTest");
+				return false;
+			}
+			return true;
+		}
+
+		TStrongObjectPtr<UPixelStreaming2DynamicDelegateTest> DelegateTest;
+	};
+
+	template <typename DelegateType, typename... Args>
+	TSharedPtr<FSingleDelegateTest<StripOptionalType<Args>...>> CreateSingleDelegateTest(TMap<FString, TSharedPtr<UE::PixelStreaming2::FSingleDelegateTestBase>>& Map, DelegateType& InDelegate, FString InName)
+	{
+		TSharedPtr<FSingleDelegateTest<StripOptionalType<Args>...>> DelegateTest = MakeShared<FSingleDelegateTest<StripOptionalType<Args>...>>(InName);
+		DelegateTest->BindDelegate(InDelegate);
+		Map.Add(InName, DelegateTest);
+		return DelegateTest;
+	}
+
+	class FDelegateNativeTest: public FDelegateTestBase
+	{
+	public:
+		~FDelegateNativeTest()
+		{
+			for (auto& DelegateTest : DelegatesMap)
+			{
+				if (const auto& Value = DelegateTest.Value; !Value->bWasCalledExpectedTimes(true))
+				{
+					UE_LOGFMT(LogPixelStreaming2RTC, Error, "{0} was called {1} times.", Value->Name, Value->CallCount);
+				}
+			}
+		}
+
+		bool Init(DelegateTestConfig Config, FString StreamerName)
+		{
+			UPixelStreaming2Delegates* Delegates = UPixelStreaming2Delegates::Get();
+			if (!Delegates)
+			{
+				UE_LOGFMT(LogPixelStreaming2RTC, Error, "Delegates are null.");
+				return false;
+			}
+
+			const bool bIsRemote = false;
+
+			CreateSingleDelegateTest<decltype(Delegates->OnConnectedToSignallingServerNative), FString>(DelegatesMap, Delegates->OnConnectedToSignallingServerNative, TEXT("OnConnectedToSignallingServerNative"))->Times(Exactly(1)).With(StreamerName);
+			CreateSingleDelegateTest<decltype(Delegates->OnDisconnectedFromSignallingServerNative), FString>(DelegatesMap, Delegates->OnDisconnectedFromSignallingServerNative, TEXT("OnDisconnectedFromSignallingServerNative"))->Times(Exactly(1)).With(StreamerName);
+			CreateSingleDelegateTest<decltype(Delegates->OnNewConnectionNative), FString, TOptional<FString>>(DelegatesMap, Delegates->OnNewConnectionNative, TEXT("OnNewConnectionNative"))->Times(Exactly(Config.NumPlayers)).With(StreamerName, Any<FString>());
+			// TODO (Eden.Harris) This currently only fires once but should fire NumPlayers times.
+			CreateSingleDelegateTest<decltype(Delegates->OnClosedConnectionNative), FString, TOptional<FString>>(DelegatesMap, Delegates->OnClosedConnectionNative, TEXT("OnClosedConnectionNative"))->Times(Exactly(1)).With(StreamerName, Any<FString>());
+			CreateSingleDelegateTest<decltype(Delegates->OnAllConnectionsClosedNative), FString>(DelegatesMap, Delegates->OnAllConnectionsClosedNative, TEXT("OnAllConnectionsClosedNative"))->Times(Exactly(1)).With(StreamerName);
+			CreateSingleDelegateTest<decltype(Delegates->OnDataTrackOpenNative), FString, TOptional<FString>>(DelegatesMap, Delegates->OnDataTrackOpenNative, TEXT("OnDataTrackOpenNative"))->Times(Exactly(Config.NumPlayers)).With(StreamerName, Any<FString>());
+			// TODO (Eden.Harris) DataTrack closed does not currently fire
+			CreateSingleDelegateTest<decltype(Delegates->OnDataTrackClosedNative), FString, TOptional<FString>>(DelegatesMap, Delegates->OnDataTrackClosedNative, TEXT("OnDataTrackClosedNative"))->Times(AtLeast(0)).With(StreamerName, Any<FString>());
+
+			const int NumCalls = Config.bIsBidirectional ? Config.NumPlayers * 2: Config.NumPlayers;
+
+			auto& OnVideoTrackOpenNativeDelegate = CreateSingleDelegateTest<decltype(Delegates->OnVideoTrackOpenNative), FString, TOptional<FString>, bool>(DelegatesMap, Delegates->OnVideoTrackOpenNative, TEXT("OnVideoTrackOpenNative"))->Times(Exactly(NumCalls)).With(StreamerName, Any<FString>(), bIsRemote);
+			// TODO (Eden.Harris) This should be called twice, once for local and remote. 
+			// This is also not being fired on linux or Mac so is disabled for now.
+			// auto& OnVideoTrackClosedNativeDelegate = CreateSingleDelegateTest<decltype(Delegates->OnVideoTrackClosedNative), FString, TOptional<FString>, bool>(DelegatesMap, Delegates->OnVideoTrackClosedNative, TEXT("OnVideoTrackClosedNative"))->Times(Exactly(Config.NumPlayers)).With(StreamerName, Any<FString>(), bIsRemote);
+			auto& OnAudioTrackOpenNativeDelegate = CreateSingleDelegateTest<decltype(Delegates->OnAudioTrackOpenNative), FString, TOptional<FString>, bool>(DelegatesMap, Delegates->OnAudioTrackOpenNative, TEXT("OnAudioTrackOpenNative"))->Times(Exactly(NumCalls)).With( StreamerName, Any<FString>(), bIsRemote);
+			// TODO (Eden.Harris) This should be called twice, once for local and remote
+			auto& OnAudioTrackClosedNativeDelegate = CreateSingleDelegateTest<decltype(Delegates->OnAudioTrackClosedNative), FString, TOptional<FString>, bool>(DelegatesMap, Delegates->OnAudioTrackClosedNative, TEXT("OnAudioTrackClosedNative"))->Times(Exactly(Config.NumPlayers)).With(StreamerName, Any<FString>(), bIsRemote);
+			
+			if(Config.bIsBidirectional)
+			{
+				OnVideoTrackOpenNativeDelegate.With(StreamerName, Any<FString>(), !bIsRemote);
+				//OnVideoTrackClosedNativeDelegate.With(StreamerName, Any<FString>(), !bIsRemote);
+				OnAudioTrackOpenNativeDelegate.With(StreamerName, Any<FString>(), !bIsRemote);
+				OnAudioTrackClosedNativeDelegate.With(StreamerName, Any<FString>(), !bIsRemote);
+			}
+
+			CreateSingleDelegateTest<decltype(Delegates->OnStatChangedNative), TOptional<FString>, TOptional<FName>, TOptional<float>>(DelegatesMap, Delegates->OnStatChangedNative, TEXT("OnStatChangedNative"))->Times(AtLeast(1)).With(Any<FString>(), Any<FName>(), Any<float>());
+			CreateSingleDelegateTest<decltype(Delegates->OnFallbackToSoftwareEncodingNative)>(DelegatesMap, Delegates->OnFallbackToSoftwareEncodingNative, TEXT("OnFallbackToSoftwareEncodingNative"))->Times(Exactly(Config.SoftwareEncodingCount));
+
+			return true;
+		}
+	};
+
+	DEFINE_LATENT_AUTOMATION_COMMAND_TWO_PARAMETER(FCleanupDelegatesNative, TSharedPtr<FDelegateNativeTest>, DelegatesTest, float, TimeoutSeconds);
+
+	bool FCleanupDelegatesNative::Update()
+	{
+		const double DeltaTime = FPlatformTime::Seconds() - StartTime;
+		if (DeltaTime > TimeoutSeconds)
+		{
+			UE_LOG(LogPixelStreaming2RTC, Error, TEXT("Timed out waiting for delegates."));
+			return true;
+		}
+		
+		if (DelegatesTest && DelegatesTest->CheckCalled(false))
+		{
+			DelegatesTest.Reset();
+			UE_LOGFMT(LogPixelStreaming2RTC, Log, "Cleaning up DelegatesTest.");
+			return true;
+		}
+		else if (!DelegatesTest)
+		{
+			UE_LOGFMT(LogPixelStreaming2RTC, Error, "DelegatesTest is null.");
+			return true;
+		}
+		return false;
+	}
+
+	DEFINE_LATENT_AUTOMATION_COMMAND_TWO_PARAMETER(FCleanupDelegates, TSharedPtr<FDynamicDelegateLifetime>, DelegateTestScope, float, TimeoutSeconds);
+
+	bool FCleanupDelegates::Update()
+	{
+		const double DeltaTime = FPlatformTime::Seconds() - StartTime;
+		if (DeltaTime > TimeoutSeconds)
+		{
+			UE_LOG(LogPixelStreaming2RTC, Error, TEXT("Timed out waiting for delegates."));
+			return true;
+		}
+
+		if (DelegateTestScope->DelegateTest && DelegateTestScope->DelegateTest->CheckCalled(false))
+		{
+			UE_LOGFMT(LogPixelStreaming2RTC, Log, "Cleaning up DelegatesTest.");
+			return true;
+		}
+
+		return false;
+	}
+
+	template<typename T>
+	void RunDelegateTest(DelegateTestConfig Config)
+	{
+		const int32 StreamerPort = TestUtils::NextStreamerPort();
+		const int32 PlayerPort = TestUtils::NextPlayerPort();
+		const FString StreamerName(FString::Printf(TEXT("MockStreamer%d"), StreamerPort));
+		
+		UPixelStreaming2Delegates* Delegates = UPixelStreaming2Delegates::Get();
+		if (!Delegates)
+		{
+			UE_LOGFMT(LogPixelStreaming2RTC, Error, "Delegates are null.");
+			return;
+		}
+		
+		const TSharedPtr<T> DelegatesTest = MakeShared<T>();
+		if (!DelegatesTest->Init(Config, StreamerName))
+		{
+			UE_LOGFMT(LogPixelStreaming2RTC, Error, "Unable to create FDelegatesTest");
+			return;
+		}
+
+		const TSharedPtr<UE::PixelStreaming2Servers::IServer> SignallingServer = CreateSignallingServer(StreamerPort, PlayerPort);
+
+		const TSharedPtr<IPixelStreaming2Streamer> Streamer = CreateStreamer(StreamerName, StreamerPort);
+		const TSharedPtr<FVideoProducer>		   VideoProducer = FVideoProducer::Create();
+		Streamer->SetVideoProducer(VideoProducer);
+
+		TArray<TSharedPtr<FMockPlayer>> Players;
+		Players.SetNum(Config.NumPlayers);
+		for (TSharedPtr<FMockPlayer>& Player : Players)
+		{
+			FMockPlayerConfig PlayerConfig = {};
+			if(Config.bIsBidirectional)
+			{
+				PlayerConfig.AudioDirection = EMediaDirection::Bidirectional;
+				PlayerConfig.VideoDirection = EMediaDirection::Bidirectional;
+			}
+			Player = CreatePlayer(PlayerConfig);
+		}
+		TArray<TSharedPtr<FMockVideoSink>> VideoSinks;
+		VideoSinks.SetNum(Config.NumPlayers);
+		for (int i = 0; i < VideoSinks.Num(); ++i)
+		{
+			VideoSinks[i] = Players[i]->GetVideoSink();
+		}
+
+		const TSharedPtr<bool> bStreamingStarted = MakeShared<bool>(false);
+		Streamer->OnStreamingStarted().AddLambda([bStreamingStarted](IPixelStreaming2Streamer*) {
+			*(bStreamingStarted.Get()) = true;
+		});
+
+		const TSharedPtr<bool> bStreamingDisconnected = MakeShared<bool>(false);
+		Delegates->OnDisconnectedFromSignallingServerNative.AddLambda([bStreamingDisconnected](FString) {
+			*(bStreamingDisconnected.Get()) = true;
+		});
+
+		ADD_LATENT_AUTOMATION_COMMAND(FExecuteLambda([Streamer]() { Streamer->StartStreaming(); }))
+		ADD_LATENT_AUTOMATION_COMMAND(FWAitForBoolOrTimeout(TEXT("Check streaming started"), 5.0, Streamer, bStreamingStarted, true))
+
+		for (TSharedPtr<FMockPlayer>& Player : Players)
+		{
+			ADD_LATENT_AUTOMATION_COMMAND(FExecuteLambda([Player, PlayerPort]() { Player->Connect(PlayerPort); }))
+			ADD_LATENT_AUTOMATION_COMMAND(FSubscribePlayerAfterStreamerConnectedOrTimeout(5.0, Streamer, Player, StreamerName))
+		}
+
+		for (TSharedPtr<FMockPlayer>& Player : Players)
+		{
+			ADD_LATENT_AUTOMATION_COMMAND(FWaitForDataChannelOrTimeout(5.0, Player));
+		}
+
+		// Wait 1 second to ensure any websocket message have correctly flowed
+		ADD_LATENT_AUTOMATION_COMMAND(FWaitSeconds(1.0))
+		ADD_LATENT_AUTOMATION_COMMAND(FExecuteLambda([Streamer]() { Streamer->StopStreaming(); }))
+		ADD_LATENT_AUTOMATION_COMMAND(FWAitForBoolOrTimeout(TEXT("Check disconnected"), 5.0, Streamer, bStreamingDisconnected, true))
+
+		ADD_LATENT_AUTOMATION_COMMAND(FCleanupAllPlayers(SignallingServer, Streamer, Players))
+		
+		if constexpr (std::is_same_v<T, FDelegateNativeTest>)
+		{
+			ADD_LATENT_AUTOMATION_COMMAND(FCleanupDelegatesNative(DelegatesTest, 5.0))
+		}
+		else if constexpr (std::is_same_v<T, FDynamicDelegateLifetime>)
+		{
+			ADD_LATENT_AUTOMATION_COMMAND(FCleanupDelegates(DelegatesTest, 5.0))
+		}
+	}
+
+	// TODO (Eden.Harris) RTCP-8326 This test is failing to fires some delegates.
+#if 0
+	IMPLEMENT_SIMPLE_AUTOMATION_TEST(FPS2DelegateSoftwareFallbackTest, "System.Plugins.PixelStreaming2.FPS2DelegateSoftwareFallbackTest", EAutomationTestFlags::EditorContext | EAutomationTestFlags::ClientContext | EAutomationTestFlags::ProductFilter)
+	bool FPS2DelegateSoftwareFallbackTest::RunTest(const FString&)
+	{
+		bool Result = true;
+		if(IsRHIDeviceNVIDIA())
+		{
+			AddExpectedError(TEXT("No more HW encoders available. Falling back to software encoding"), EAutomationExpectedMessageFlags::MatchType::Exact, 1, false);
+			const int PrevPixelStreamingEncoderMaxSessions = UPixelStreaming2PluginSettings::CVarEncoderMaxSessions.GetValueOnAnyThread();
+			const EVideoCodec PrevCodec = UE::PixelStreaming2::GetEnumFromCVar<EVideoCodec>(UPixelStreaming2PluginSettings::CVarEncoderCodec);
+			UPixelStreaming2PluginSettings::CVarEncoderMaxSessions.AsVariable()->Set(0);
+			UPixelStreaming2PluginSettings::CVarEncoderCodec.AsVariable()->Set(*UE::PixelStreaming2::GetCVarStringFromEnum(EVideoCodec::H264));
+			Result = RunDelegateTest<FDelegateNativeTest>(1, 1);
+
+			// Reset to previous
+			ADD_LATENT_AUTOMATION_COMMAND(FExecuteLambda([PrevPixelStreamingEncoderMaxSessions, PrevCodec]() { 
+				UPixelStreaming2PluginSettings::CVarEncoderMaxSessions.AsVariable()->Set(PrevPixelStreamingEncoderMaxSessions);
+				UPixelStreaming2PluginSettings::CVarEncoderCodec.AsVariable()->Set(*UE::PixelStreaming2::GetCVarStringFromEnum(PrevCodec));
+			}))
+		}
+		else
+		{
+			UE_LOGFMT(LogPixelStreaming2RTC, Log, "FPS2DelegateSoftwareFallbackTest requires Nvidia GPU to test Software fallback");
+		}
+		
+		return Result;
+	}
+#endif
+
+	IMPLEMENT_SIMPLE_AUTOMATION_TEST(FPS2DelegateNativeSingleTest, "System.Plugins.PixelStreaming2.FPS2DelegateNativeSingleTest", EAutomationTestFlags::EditorContext | EAutomationTestFlags::ClientContext | EAutomationTestFlags::ProductFilter)
+	bool FPS2DelegateNativeSingleTest::RunTest(const FString&)
+	{
+		RunDelegateTest<FDelegateNativeTest>({0, 1, false});
+		return true;
+	}
+
+	IMPLEMENT_SIMPLE_AUTOMATION_TEST(FPS2DelegateDynamicSingleTest, "System.Plugins.PixelStreaming2.FPS2DelegateDynamicSingleTest", EAutomationTestFlags::EditorContext | EAutomationTestFlags::ClientContext | EAutomationTestFlags::ProductFilter)
+	bool FPS2DelegateDynamicSingleTest::RunTest(const FString&)
+	{
+		RunDelegateTest<FDynamicDelegateLifetime>({0, 1, false});
+		return true;
+	}
+
+	IMPLEMENT_SIMPLE_AUTOMATION_TEST(FPS2DelegateNativeMultipleTest, "System.Plugins.PixelStreaming2.FPS2DelegateNativeMultipleTest", EAutomationTestFlags::EditorContext | EAutomationTestFlags::ClientContext | EAutomationTestFlags::ProductFilter)
+	bool FPS2DelegateNativeMultipleTest::RunTest(const FString&)
+	{
+		RunDelegateTest<FDelegateNativeTest>({0, 3, false});
+		return true;
+	}
+
+	IMPLEMENT_SIMPLE_AUTOMATION_TEST(FPS2DelegateDynamicMultipleTest, "System.Plugins.PixelStreaming2.FPS2DelegateDynamicMultipleTest", EAutomationTestFlags::EditorContext | EAutomationTestFlags::ClientContext | EAutomationTestFlags::ProductFilter)
+	bool FPS2DelegateDynamicMultipleTest::RunTest(const FString&)
+	{
+		RunDelegateTest<FDynamicDelegateLifetime>({0, 3, false});
+		return true;
+	}
+
+	IMPLEMENT_SIMPLE_AUTOMATION_TEST(FPS2DelegateNativeSingleBidirectionalTest, "System.Plugins.PixelStreaming2.FPS2DelegateNativeSingleBidirectionalTest", EAutomationTestFlags::EditorContext | EAutomationTestFlags::ClientContext | EAutomationTestFlags::ProductFilter)
+	bool FPS2DelegateNativeSingleBidirectionalTest::RunTest(const FString&)
+	{
+		RunDelegateTest<FDelegateNativeTest>({0, 1, true});
+		return true;
+	}
+
+	IMPLEMENT_SIMPLE_AUTOMATION_TEST(FPS2DelegateDynamicSingleBidirectionalTest, "System.Plugins.PixelStreaming2.FPS2DelegateDynamicSingleBidirectionalTest", EAutomationTestFlags::EditorContext | EAutomationTestFlags::ClientContext | EAutomationTestFlags::ProductFilter)
+	bool FPS2DelegateDynamicSingleBidirectionalTest::RunTest(const FString&)
+	{
+		RunDelegateTest<FDynamicDelegateLifetime>({0, 1, true});
+		return true;
+	}
+
+	IMPLEMENT_SIMPLE_AUTOMATION_TEST(FPS2DelegateNativeMultipleBidirectionalTest, "System.Plugins.PixelStreaming2.FPS2DelegateNativeMultipleBidirectionalTest", EAutomationTestFlags::EditorContext | EAutomationTestFlags::ClientContext | EAutomationTestFlags::ProductFilter)
+	bool FPS2DelegateNativeMultipleBidirectionalTest::RunTest(const FString&)
+	{
+		RunDelegateTest<FDelegateNativeTest>({0, 3, true});
+		return true;
+	}
+
+	IMPLEMENT_SIMPLE_AUTOMATION_TEST(FPS2DelegateDynamicMultipleBidirectionalTest, "System.Plugins.PixelStreaming2.FPS2DelegateDynamicMultipleBidirectionalTest", EAutomationTestFlags::EditorContext | EAutomationTestFlags::ClientContext | EAutomationTestFlags::ProductFilter)
+	bool FPS2DelegateDynamicMultipleBidirectionalTest::RunTest(const FString&)
+	{
+		RunDelegateTest<FDynamicDelegateLifetime>({0, 3, true});
+		return true;
+	}
+} // namespace UE::PixelStreaming2
+
+#endif // WITH_TESTS
diff --git a/Engine/Plugins/Media/PixelStreaming2/Source/PixelStreaming2RTC/Private/Tests/DelegateTest.h b/Engine/Plugins/Media/PixelStreaming2/Source/PixelStreaming2RTC/Private/Tests/DelegateTest.h
new file mode 100644
--- /dev/null
+++ b/Engine/Plugins/Media/PixelStreaming2/Source/PixelStreaming2RTC/Private/Tests/DelegateTest.h
@@ -0,0 +1,329 @@
+// Copyright Epic Games, Inc. All Rights Reserved.
+
+#pragma once
+
+#include "Logging.h"
+#include "PixelStreaming2Delegates.h"
+#include "UObject/Object.h"
+
+#include "DelegateTest.generated.h"
+
+#if WITH_TESTS
+
+namespace UE::PixelStreaming2
+{
+	template <typename T>
+	TOptional<T> Any()
+	{
+		TOptional<T> Temp;
+		return Temp;
+	}
+
+	template<typename T>
+	struct IsOptional : std::false_type {};
+
+	template<typename T>
+	struct IsOptional<TOptional<T>> : std::true_type {};
+
+	template<typename T>
+	struct StripOptional{using type = T;};
+
+	template<typename T>
+	struct StripOptional<TOptional<T>>{using type = T;};
+
+	template<typename T>
+	using StripOptionalType = typename StripOptional<T>::type;
+
+	template<typename T>
+	auto ToOptional(T&& a)
+	{
+		if constexpr(IsOptional<std::decay_t<T>>::value)
+		{
+			return Forward<T>(a);
+		}
+		else
+		{
+			return TOptional<std::decay_t<T>>(Forward<T>(a));
+		}
+	}
+	
+	class FCardinality
+	{
+	public:
+		FCardinality()
+			: Min(std::numeric_limits<int>::min())
+			, Max(std::numeric_limits<int>::max())
+		{
+		}
+
+		FCardinality(int Min, int Max)
+			: Min(Min)
+			, Max(Max)
+		{
+		}
+
+		int Min;
+		int Max;
+	};
+
+	FCardinality AnyNumber();
+	FCardinality AtLeast(int Min);
+	FCardinality AtMost(int Max);
+	FCardinality Between(int Min, int Max);
+	FCardinality Exactly(int ExactValue);
+
+	struct DelegateTestConfig
+	{
+		int SoftwareEncodingCount = 0;
+		int NumPlayers = 0;
+		bool bIsBidirectional = false;
+	};
+
+	// Base class to allow for Delegate Tests to be stored in an array with different parameters
+	class FSingleDelegateTestBase
+	{
+	public:
+		FSingleDelegateTestBase(FString InName): 
+			Name(MoveTemp(InName))
+			{}
+
+		virtual ~FSingleDelegateTestBase() = default;
+
+		bool bWasCalledExpectedTimes(bool bPrintErrors) const
+		{
+			const bool bGreaterThanMin = CallCount >= ExpectedCallCount.Min;
+			const bool bLessThanMax = CallCount <= ExpectedCallCount.Max;
+
+			if (bPrintErrors && (!bGreaterThanMin || !bLessThanMax))
+			{
+				if (!bGreaterThanMin)
+				{
+					UE_LOGFMT(LogPixelStreaming2RTC, Error, "{0} Expected Value {1} not greater than {2}", Name, ExpectedCallCount.Min, CallCount);
+				}
+
+				if (!bLessThanMax)
+				{
+					UE_LOGFMT(LogPixelStreaming2RTC, Error, "{0} Expected Value {1} not less than {2}", Name, ExpectedCallCount.Max, CallCount);
+				}
+			}
+
+			return CallCount >= ExpectedCallCount.Min
+				&& CallCount <= ExpectedCallCount.Max
+				&& bCallbackMatchesExpectedValues;
+		}
+
+		FString Name;
+		int			 CallCount = 0;
+		FCardinality ExpectedCallCount;
+		bool bCallbackMatchesExpectedValues = true;
+	};
+
+	template <typename... Args>
+	class FSingleDelegateArgsTest: public FSingleDelegateTestBase
+	{
+	public:
+		FSingleDelegateArgsTest(FString InName)
+			: FSingleDelegateTestBase(MoveTemp(InName)) 
+			{}
+
+		void OnCalled(Args... InActualValues)
+		{
+			++CallCount;
+
+			auto ActualValues = MakeTuple(InActualValues...);
+
+			const int NumExpectedValues = ExpectedValuesArray.Num();
+
+			// If there's no expected values, the check is an automatic success
+			bool bCheckSuccess = NumExpectedValues == 0; 
+			for (int i = (NumExpectedValues - 1); i >= 0; i--)
+			{
+				TArray<FString> CheckStrings;
+
+				auto& ExpectedValueSet = ExpectedValuesArray[i];
+				bool  bThisCheckSuccess = true;
+
+				VisitTupleElements([&bThisCheckSuccess, &CheckStrings, Name = Name](auto& ExpectedValue, auto& ActualValue) {
+
+					if (!ExpectedValue.IsSet())
+					{
+						bThisCheckSuccess &= true;
+					}
+					else
+					{
+						bThisCheckSuccess &= ExpectedValue == ActualValue;
+					}
+				}, ExpectedValueSet, ActualValues);
+
+				// This check has succeeded so no need to check earlier registered expected values
+				if (bThisCheckSuccess)
+				{
+					bCheckSuccess = true;
+					break;
+				}
+			}
+
+			if (!bCheckSuccess)
+			{
+				UE_LOGFMT(LogPixelStreaming2RTC, Error, "{0} expected Value do not match actual values", Name);
+			}
+
+			bCallbackMatchesExpectedValues &= bCheckSuccess;
+		}
+
+		FSingleDelegateArgsTest<Args...>& Times(const FCardinality& InExpectedCallCount)
+		{
+			ExpectedCallCount = InExpectedCallCount;
+			return *this;
+		}
+
+		// Use Separate template args because the types may include TOptional while args does not.
+		template<typename... Params>
+		FSingleDelegateArgsTest<Args...>& With(Params... Values)
+		{
+			static_assert(sizeof...(Params) == sizeof...(Args));
+			
+			// Make sure all parameters get wrapped in a TOptional if they are not already. 
+			// This is to get around issues with embedding an existing TOptional inside another TOptional.
+			// When double embedded, the tests will fail because the outer TOptional thinks it is set when inner is empty.
+			ExpectedValuesArray.Add(MakeTuple(ToOptional(Forward<Params>(Values))...));
+			return *this;
+		}
+
+		TArray<TTuple<TOptional<Args>...>> ExpectedValuesArray;
+	};
+
+	template <typename... Args>
+	class FSingleDynamicDelegateTest: public FSingleDelegateArgsTest<Args...>
+	{
+	public:
+		FSingleDynamicDelegateTest(FString InName)
+			: FSingleDelegateArgsTest<Args...>(InName) 
+			{}
+	};
+
+	template <typename... Args>
+	class FSingleDelegateTest : public FSingleDelegateArgsTest<Args...>, public TSharedFromThis<FSingleDelegateTest<Args...>>
+	{
+	public:
+		FSingleDelegateTest(FString InName)
+			: FSingleDelegateArgsTest<Args...>(MoveTemp(InName)) 
+			{}
+	
+		virtual ~FSingleDelegateTest()
+		{
+			if (UnbindDelegateFunc)
+			{
+				UnbindDelegateFunc();
+			}
+		}
+
+		template<typename DelegateType>
+		void BindDelegate(DelegateType& InDelegate)
+		{
+			DelegateHandle = InDelegate.AddLambda([WeakThis = this->AsWeak()](auto... InActualValues)
+			{
+				if (const TSharedPtr<FSingleDelegateTest> SharedThis = WeakThis.Pin())
+				{
+					auto ActualValues = MakeTuple(InActualValues...);
+					SharedThis->OnCalled(InActualValues...);
+				}
+			});
+
+			UnbindDelegateFunc = [&InDelegate, Handle = DelegateHandle]()
+			{
+				InDelegate.Remove(Handle);
+			};
+		}
+
+		TFunction<void()> UnbindDelegateFunc;
+		FDelegateHandle DelegateHandle;
+	};
+
+	class FDelegateTestBase
+	{
+	public:
+		bool CheckCalled(bool bPrintErrors) const;
+
+		TMap<FString, TSharedPtr<UE::PixelStreaming2::FSingleDelegateTestBase>> DelegatesMap;
+	};
+}
+
+// Blueprint dynamic delegates require UE's reflection system.
+UCLASS()
+class UPixelStreaming2DynamicDelegateTest : public UObject, public UE::PixelStreaming2::FDelegateTestBase
+{
+	GENERATED_BODY()
+public:
+	UFUNCTION()
+	void OnConnectedToSignallingServer(FString StreamerId);
+
+	UFUNCTION()
+	void OnDisconnectedFromSignallingServer(FString StreamerId);
+
+	UFUNCTION()
+	void OnNewConnection(FString StreamerId, FString PlayerId);
+
+	UFUNCTION()
+	void OnClosedConnection(FString StreamerId, FString PlayerId);
+
+	UFUNCTION()
+	void OnAllConnectionsClosed(FString StreamerId);
+
+	UFUNCTION()
+	void OnDataTrackOpen(FString StreamerId, FString PlayerId);
+
+	UFUNCTION()
+	void OnDataTrackClosed(FString StreamerId, FString PlayerId);
+
+	UFUNCTION()
+	void OnStatChanged(FString PlayerId, FName StatName, float StatValue);
+
+	UFUNCTION()
+	void OnFallbackToSoftwareEncoding();
+
+	bool Init(UE::PixelStreaming2::DelegateTestConfig Config, FString StreamerName);
+	void Destroy();
+
+	template <typename... Args>
+	TSharedPtr<UE::PixelStreaming2::FSingleDelegateArgsTest<UE::PixelStreaming2::StripOptionalType<Args>...>> BindDelegate(FString InName)
+	{
+		UPixelStreaming2Delegates* Delegates = UPixelStreaming2Delegates::Get();
+		if (!Delegates)
+		{
+			UE_LOGFMT(LogPixelStreaming2RTC, Error, "Delegates are null.");
+			return nullptr;
+		}
+		else
+		{
+			TSharedPtr<UE::PixelStreaming2::FSingleDelegateArgsTest<UE::PixelStreaming2::StripOptionalType<Args>...>> Ptr = MakeShared<UE::PixelStreaming2::FSingleDynamicDelegateTest<UE::PixelStreaming2::StripOptionalType<Args>...>>(InName);
+			DelegatesMap.Add(InName, Ptr);
+			return Ptr;
+		}
+	}
+
+private:
+	
+	template <typename... Args>
+	void DynamicDelegateCalled(FString InName, Args... InActualValues)
+	{
+		UE_LOGFMT(LogPixelStreaming2RTC, VeryVerbose, "{0} was called", InName);
+		if (!IsInGameThread())
+		{
+			UE_LOGFMT(LogPixelStreaming2RTC, Error, "{0} was not called on the game thread", InName);
+		}
+		
+		TSharedPtr<UE::PixelStreaming2::FSingleDelegateTestBase>* DelegateTest = DelegatesMap.Find(InName);
+		if (DelegateTest)
+		{
+			TSharedPtr<UE::PixelStreaming2::FSingleDynamicDelegateTest<Args...>> DynamicTest = StaticCastSharedPtr<UE::PixelStreaming2::FSingleDynamicDelegateTest<Args...>>(*DelegateTest);
+			DynamicTest->OnCalled(InActualValues...);
+		}
+		else
+		{
+			UE_LOGFMT(LogPixelStreaming2RTC, Error, "unknown Delegate Test {0}", InName);
+		}
+	}
+};
+
+#endif // WITH_TESTS
diff --git a/Engine/Plugins/Media/PixelStreaming2/Source/PixelStreaming2RTC/Private/Tests/MockPlayer.cpp b/Engine/Plugins/Media/PixelStreaming2/Source/PixelStreaming2RTC/Private/Tests/MockPlayer.cpp
--- a/Engine/Plugins/Media/PixelStreaming2/Source/PixelStreaming2RTC/Private/Tests/MockPlayer.cpp
+++ b/Engine/Plugins/Media/PixelStreaming2/Source/PixelStreaming2RTC/Private/Tests/MockPlayer.cpp
@@ -46,7 +46,7 @@ namespace UE::PixelStreaming2
 		Player->EpicRtcVideoEncoderInitializers = { new FEpicRtcVideoEncoderInitializer() };
 		Player->EpicRtcVideoDecoderInitializers = { new FEpicRtcVideoDecoderInitializer() };
 
-		FUtf8String ConferenceId = "test_conference";
+		FUtf8String ConferenceId = FUtf8String::Printf("test_conference_%d", PlayerId);
 
 		EpicRtcErrorCode Result = GetOrCreatePlatform({}, Player->Platform.GetInitReference());
 
diff --git a/Engine/Plugins/Media/PixelStreaming2/Source/PixelStreaming2RTC/Private/Tests/TestUtils.cpp b/Engine/Plugins/Media/PixelStreaming2/Source/PixelStreaming2RTC/Private/Tests/TestUtils.cpp
--- a/Engine/Plugins/Media/PixelStreaming2/Source/PixelStreaming2RTC/Private/Tests/TestUtils.cpp
+++ b/Engine/Plugins/Media/PixelStreaming2/Source/PixelStreaming2RTC/Private/Tests/TestUtils.cpp
@@ -206,7 +206,7 @@ namespace UE::PixelStreaming2
 		double DeltaTime = FPlatformTime::Seconds() - StartTime;
 		if (DeltaTime > TimeoutSeconds)
 		{
-			UE_LOG(LogPixelStreaming2RTC, Error, TEXT("Timed out waiting to subscribe player."));
+			UE_LOG(LogPixelStreaming2RTC, Error, TEXT("Timed out waiting to data channel to be available."));
 			return true;
 		}
 		return false; // Not connected or timed out so run this latent test again next frame
@@ -248,6 +248,25 @@ namespace UE::PixelStreaming2
 		return false; // We keep updating this until the timeout is complete and then we check that there's no tracks added
 	}
 
+	bool FWAitForBoolOrTimeout::Update()
+	{
+		double DeltaTime = FPlatformTime::Seconds() - StartTime;
+		if (DeltaTime > WaitSeconds)
+		{
+			if ((*bCheck.Get()) != bExpectedValue)
+			{
+				UE_LOGFMT(LogPixelStreaming2RTC, Error, "{0} failed. Expected [{1}] but got [{2}]", CheckName, bExpectedValue, *bCheck.Get());
+			}
+			return true;
+		}
+
+		if ((*bCheck.Get()) == bExpectedValue)
+		{
+			return true;
+		}
+		return false; // We keep updating this until the timeout is complete and then we check that there's no tracks added
+	}
+
 	bool FWaitAndCheckStreamerBool::Update()
 	{
 		double DeltaTime = FPlatformTime::Seconds() - StartTime;
@@ -314,6 +333,34 @@ namespace UE::PixelStreaming2
 		return true;
 	}
 
+	bool FCleanupAllPlayers::Update()
+	{
+		for(TSharedPtr<FMockPlayer>& Player : OutPlayers)
+		{
+			if (Player)
+			{
+				Player.Reset();
+			}
+		}
+
+		if (OutStreamer)
+		{
+			OutStreamer->StopStreaming();
+			OutStreamer.Reset();
+		}
+
+		if (OutSignallingServer)
+		{
+			OutSignallingServer->Stop();
+			OutSignallingServer.Reset();
+		}
+
+		// Restore media directions back to default
+		SetMediaDirection(EMediaType::Audio, EMediaDirection::Bidirectional);
+		SetMediaDirection(EMediaType::Video, EMediaDirection::Bidirectional);
+		return true;
+	}
+
 	bool FExecuteLambda::Update()
 	{
 		Func();
diff --git a/Engine/Plugins/Media/PixelStreaming2/Source/PixelStreaming2RTC/Private/Tests/TestUtils.h b/Engine/Plugins/Media/PixelStreaming2/Source/PixelStreaming2RTC/Private/Tests/TestUtils.h
--- a/Engine/Plugins/Media/PixelStreaming2/Source/PixelStreaming2RTC/Private/Tests/TestUtils.h
+++ b/Engine/Plugins/Media/PixelStreaming2/Source/PixelStreaming2RTC/Private/Tests/TestUtils.h
@@ -80,6 +80,7 @@ namespace UE::PixelStreaming2
 	DEFINE_LATENT_AUTOMATION_COMMAND_THREE_PARAMETER(FWaitForFrameReceived, double, TimeoutSeconds, TSharedPtr<FMockVideoSink>, VideoSink, FMockVideoFrameConfig, FrameConfig);
 	DEFINE_LATENT_AUTOMATION_COMMAND_FOUR_PARAMETER(FSubscribePlayerAfterStreamerConnectedOrTimeout, double, TimeoutSeconds, TSharedPtr<IPixelStreaming2Streamer>, OutStreamer, TSharedPtr<FMockPlayer>, OutPlayer, FString, StreamerName);
 	DEFINE_LATENT_AUTOMATION_COMMAND_THREE_PARAMETER(FCleanupAll, TSharedPtr<UE::PixelStreaming2Servers::IServer>, OutSignallingServer, TSharedPtr<IPixelStreaming2Streamer>, OutStreamer, TSharedPtr<FMockPlayer>, OutPlayer);
+	DEFINE_LATENT_AUTOMATION_COMMAND_THREE_PARAMETER(FCleanupAllPlayers, TSharedPtr<UE::PixelStreaming2Servers::IServer>, OutSignallingServer, TSharedPtr<IPixelStreaming2Streamer>, OutStreamer, TArray<TSharedPtr<FMockPlayer>>, OutPlayers);
 	DEFINE_LATENT_AUTOMATION_COMMAND_TWO_PARAMETER(FWaitForDataChannelOrTimeout, double, TimeoutSeconds, TSharedPtr<FMockPlayer>, OutPlayer);
 	DEFINE_LATENT_AUTOMATION_COMMAND_THREE_PARAMETER(FSendDataChannelMessageToStreamer, TSharedPtr<FMockPlayer>, Player, FString, MessageType, const FString, Body);
 	DEFINE_LATENT_AUTOMATION_COMMAND_THREE_PARAMETER(FSendDataChannelMessageFromStreamer, TSharedPtr<IPixelStreaming2Streamer>, Streamer, FString, MessageType, const FString, Body);
@@ -87,6 +88,7 @@ namespace UE::PixelStreaming2
 	DEFINE_LATENT_AUTOMATION_COMMAND_THREE_PARAMETER(FWaitForPlayerTrackOrTimeout, double, TimeoutSeconds, TSharedPtr<FMockPlayer>, Player, TSharedPtr<bool>, bComplete);
 	DEFINE_LATENT_AUTOMATION_COMMAND_FOUR_PARAMETER(FWaitAndCheckBool, double, WaitSeconds, TSharedPtr<FMockPlayer>, Player, TSharedPtr<bool>, bCheck, bool, bExpectedValue);
 	DEFINE_LATENT_AUTOMATION_COMMAND_FIVE_PARAMETER(FWaitAndCheckStreamerBool, const FString, CheckName, double, WaitSeconds, TSharedPtr<IPixelStreaming2Streamer>, Streamer, TSharedPtr<bool>, bCheck, bool, bExpectedValue);
+	DEFINE_LATENT_AUTOMATION_COMMAND_FIVE_PARAMETER(FWAitForBoolOrTimeout, const FString, CheckName, double, WaitSeconds, TSharedPtr<IPixelStreaming2Streamer>, Streamer, TSharedPtr<bool>, bCheck, bool, bExpectedValue);
 	DEFINE_LATENT_AUTOMATION_COMMAND_ONE_PARAMETER(FExecuteLambda, const TFunction<void()>, Func);
 
 	TSharedPtr<IPixelStreaming2Streamer>			CreateStreamer(const FString& StreamerName, int StreamerPort);
