[PS2] fix FindPlayerBy * Track. This fixes removing incorrect rack when a player disconnects causing audio to stop playing.

Fix stores a map of tracks by pointer and player id.

diff --git a/Engine/Plugins/Media/PixelStreaming2/Source/PixelStreaming2/Internal/AudioCapturer.h b/Engine/Plugins/Media/PixelStreaming2/Source/PixelStreaming2/Internal/AudioCapturer.h
--- a/Engine/Plugins/Media/PixelStreaming2/Source/PixelStreaming2/Internal/AudioCapturer.h
+++ b/Engine/Plugins/Media/PixelStreaming2/Source/PixelStreaming2/Internal/AudioCapturer.h
@@ -80,7 +80,7 @@ namespace UE::PixelStreaming2
 		/**
 		 * This is broadcast each time audio is captured. Tracks should bind to this and push the audio into the track
 		 */
-		DECLARE_EVENT_FourParams(FAudioCapturer, FOnAudioBuffer, const int16_t*, int32, int32, const int32);
+		DECLARE_TS_MULTICAST_DELEGATE_FourParams(FOnAudioBuffer, const int16_t*, int32, int32, const int32);
 		FOnAudioBuffer OnAudioBuffer;
 
 	protected:
diff --git a/Engine/Plugins/Media/PixelStreaming2/Source/PixelStreaming2RTC/Private/EpicRtcStreamer.cpp b/Engine/Plugins/Media/PixelStreaming2/Source/PixelStreaming2RTC/Private/EpicRtcStreamer.cpp
--- a/Engine/Plugins/Media/PixelStreaming2/Source/PixelStreaming2RTC/Private/EpicRtcStreamer.cpp
+++ b/Engine/Plugins/Media/PixelStreaming2/Source/PixelStreaming2RTC/Private/EpicRtcStreamer.cpp
@@ -620,19 +620,20 @@ namespace UE::PixelStreaming2
 			return;
 		}
 
-		FPlayerContext* PlayerContext = Players->Find(PlayerId);
+		const FPlayerContext* PlayerContext = Players->Find(PlayerId);
 
 		if (!PlayerContext)
 		{
 			return;
 		}
 
-		if (!PlayerContext->DataTrack)
+		TSharedPtr<FEpicRtcDataTrack> DataTrack = PlayerContext->DataTrack;
+		if (!DataTrack)
 		{
 			return;
 		}
 
-		PlayerContext->DataTrack->SendMessage(EPixelStreaming2FromStreamerMessage::VideoEncoderAvgQP, FString::FromInt((int)StatValue));
+		DataTrack->SendMessage(EPixelStreaming2FromStreamerMessage::VideoEncoderAvgQP, FString::FromInt(static_cast<int>(StatValue)));
 	}
 
 	void FEpicRtcStreamer::DeletePlayerSession(FString PlayerId)
@@ -688,6 +689,9 @@ namespace UE::PixelStreaming2
 		}
 
 		VideoSourceGroup->RemoveAllVideoSources();
+
+		AudioTrackPlayerIdMap.Empty();
+		VideoTrackPlayerIdMap.Empty();
 		Players->Empty();
 		InputControllingId = INVALID_PLAYER_ID;
 
@@ -1114,16 +1118,17 @@ namespace UE::PixelStreaming2
 			{
 				UE_LOGFMT(LogPixelStreaming2RTC, VeryVerbose, "FEpicRtcStreamer::OnRoomStateUpdate State=Left");
 
-				if (!EpicRtcSession)
+				TRefCountPtr<EpicRtcSessionInterface> PinnedSession = EpicRtcSession;
+				if (!PinnedSession)
 				{
 					break;
 				}
 
 				FUtf8String Utf8StreamerId(StreamerId);
-				EpicRtcSession->RemoveRoom(ToEpicRtcStringView(Utf8StreamerId));
+				PinnedSession->RemoveRoom(ToEpicRtcStringView(Utf8StreamerId));
 				EpicRtcRoom = nullptr;
 
-				const EpicRtcErrorCode Result = EpicRtcSession->Disconnect(ToEpicRtcStringView("Streaming Session Removed"));
+				const EpicRtcErrorCode Result = PinnedSession->Disconnect(ToEpicRtcStringView("Streaming Session Removed"));
 				if (Result == EpicRtcErrorCode::SessionDisconnected)
 				{
 					UE_LOGFMT(LogPixelStreaming2RTC, VeryVerbose, "Session disconnected cleanly.");
@@ -1383,6 +1388,7 @@ namespace UE::PixelStreaming2
 
 		if (FPlayerContext* PlayerContext = Players->Find(ParticipantId))
 		{
+			AudioTrackPlayerIdMap.Add(reinterpret_cast<uintptr_t>(AudioTrack), ParticipantId);
 			if (AudioTrack->IsRemote())
 			{
 				PlayerContext->AudioSink = FEpicRtcAudioSink::Create(AudioTrack);
@@ -1407,6 +1413,7 @@ namespace UE::PixelStreaming2
 
 		if (FPlayerContext* PlayerContext = Players->Find(ParticipantId))
 		{
+			VideoTrackPlayerIdMap.Add(reinterpret_cast<uintptr_t>(VideoTrack), ParticipantId);
 			if (VideoTrack->IsRemote())
 			{
 				PlayerContext->VideoSink = FEpicRtcVideoSink::Create(VideoTrack);
@@ -1547,14 +1554,18 @@ namespace UE::PixelStreaming2
 			return;
 		}
 
+		AudioTrackPlayerIdMap.Remove(reinterpret_cast<uintptr_t>(AudioTrack));
 		FPlayerContext* PlayerContext = Players->Find(PlayerId);
-		if (AudioTrack->IsRemote())
-		{
-			PlayerContext->AudioSink = nullptr;
-		}
-		else
+		if (PlayerContext)
 		{
-			PlayerContext->AudioSource = nullptr;
+			if (AudioTrack->IsRemote())
+			{
+				PlayerContext->AudioSink = nullptr;
+			}
+			else
+			{
+				PlayerContext->AudioSource = nullptr;
+			}
 		}
 
 		if (UPixelStreaming2Delegates* Delegates = UPixelStreaming2Delegates::Get())
@@ -1643,15 +1654,19 @@ namespace UE::PixelStreaming2
 
 		UE_LOG(LogPixelStreaming2RTC, Log, TEXT("FEpicRtcStreamer::OnVideoTrackRemoved(VideoTrack=[%s], Player=[%s])"), *VideoTrackId, *PlayerId);
 
+		VideoTrackPlayerIdMap.Remove(reinterpret_cast<uintptr_t>(VideoTrack));
 		// If we did find a player, clear its video sink/source
 		FPlayerContext* PlayerContext = Players->Find(PlayerId);
-		if (VideoTrack->IsRemote())
-		{
-			PlayerContext->VideoSink = nullptr;
-		}
-		else
+		if (PlayerContext)
 		{
-			PlayerContext->VideoSource = nullptr;
+			if (VideoTrack->IsRemote())
+			{
+				PlayerContext->VideoSink = nullptr;
+			}
+			else
+			{
+				PlayerContext->VideoSource = nullptr;
+			}
 		}
 
 		if (UPixelStreaming2Delegates* Delegates = UPixelStreaming2Delegates::Get())
@@ -1793,29 +1808,12 @@ namespace UE::PixelStreaming2
 	bool FEpicRtcStreamer::FindPlayerByAudioTrack(EpicRtcAudioTrackInterface* AudioTrack, FString& OutPlayerId)
 	{
 		OutPlayerId = "";
-		// Find the player corresponding to this data track
-		FString AudioTrackId = ToString(AudioTrack->GetId());
-		Players->ApplyUntil([this, AudioTrackId, &OutPlayerId](FString DataPlayerId, FPlayerContext& PlayerContext) {
-			if (PlayerContext.AudioSource)
-			{
-				FString TrackId = ToString(PlayerContext.AudioSource->GetTrackId());
-				if (TrackId == AudioTrackId)
-				{
-					OutPlayerId = DataPlayerId;
-					return true;
-				}
-			}
-			if (PlayerContext.AudioSink)
-			{
-				FString TrackId = ToString(PlayerContext.AudioSink->GetTrackId());
-				if (TrackId == AudioTrackId)
-				{
-					OutPlayerId = DataPlayerId;
-					return true;
-				}
-			}
-			return false;
-		});
+
+		FString* FoundPlayerId = AudioTrackPlayerIdMap.Find(reinterpret_cast<uintptr_t>(AudioTrack));
+		if (FoundPlayerId)
+		{
+			OutPlayerId = *FoundPlayerId;
+		}
 
 		return !OutPlayerId.IsEmpty();
 	}
@@ -1823,29 +1821,12 @@ namespace UE::PixelStreaming2
 	bool FEpicRtcStreamer::FindPlayerByVideoTrack(EpicRtcVideoTrackInterface* VideoTrack, FString& OutPlayerId)
 	{
 		OutPlayerId = "";
-		// Find the player corresponding to this data track
-		FString VideoTrackId = ToString(VideoTrack->GetId());
-		Players->ApplyUntil([this, VideoTrackId, &OutPlayerId](FString DataPlayerId, FPlayerContext& PlayerContext) {
-			if (PlayerContext.VideoSource)
-			{
-				FString TrackId = ToString(PlayerContext.VideoSource->GetTrackId());
-				if (TrackId == VideoTrackId)
-				{
-					OutPlayerId = DataPlayerId;
-					return true;
-				}
-			}
-			if (PlayerContext.VideoSink)
-			{
-				FString TrackId = ToString(PlayerContext.VideoSink->GetTrackId());
-				if (TrackId == VideoTrackId)
-				{
-					OutPlayerId = DataPlayerId;
-					return true;
-				}
-			}
-			return false;
-		});
+
+		FString* FoundPlayerId = VideoTrackPlayerIdMap.Find(reinterpret_cast<uintptr_t>(VideoTrack));
+		if (FoundPlayerId)
+		{
+			OutPlayerId = *FoundPlayerId;
+		}
 
 		return !OutPlayerId.IsEmpty();
 	}
diff --git a/Engine/Plugins/Media/PixelStreaming2/Source/PixelStreaming2RTC/Private/EpicRtcStreamer.h b/Engine/Plugins/Media/PixelStreaming2/Source/PixelStreaming2RTC/Private/EpicRtcStreamer.h
--- a/Engine/Plugins/Media/PixelStreaming2/Source/PixelStreaming2RTC/Private/EpicRtcStreamer.h
+++ b/Engine/Plugins/Media/PixelStreaming2/Source/PixelStreaming2RTC/Private/EpicRtcStreamer.h
@@ -147,6 +147,8 @@ namespace UE::PixelStreaming2
 		TSharedPtr<FFreezeFrame>		  FreezeFrame;
 
 		TMap<FName, FString> ConfigOptions;
+		TMap<uintptr_t /* EpicRtcAudioTrackInterface* */,  FString /*participant id*/> AudioTrackPlayerIdMap;
+		TMap<uintptr_t /* EpicRtcVideoTrackInterface* */,  FString /*participant id*/> VideoTrackPlayerIdMap;
 
 		TSharedPtr<FStreamerReconnectTimer> ReconnectTimer;
 
