[PS2] Fix: Crashes with UObject variant observers

diff --git a/Engine/Plugins/Media/PixelStreaming2/Source/PixelStreaming2RTC/Internal/EpicRtcObserver.h b/Engine/Plugins/Media/PixelStreaming2/Source/PixelStreaming2RTC/Internal/EpicRtcObserver.h
--- a/Engine/Plugins/Media/PixelStreaming2/Source/PixelStreaming2RTC/Internal/EpicRtcObserver.h
+++ b/Engine/Plugins/Media/PixelStreaming2/Source/PixelStreaming2RTC/Internal/EpicRtcObserver.h
@@ -5,8 +5,7 @@
 #include "Misc/TVariant.h"
 #include "Templates/SharedPointer.h"
 #include "UObject/Interface.h"
-#include "UObject/WeakObjectPtr.h"
-#include "UObject/WeakObjectPtrTemplates.h"
+#include "UObject/WeakInterfacePtr.h"
 
 #include <type_traits>
 
@@ -18,7 +17,7 @@ namespace UE::PixelStreaming2
 	public:
 		TObserverVariant() = default;
 
-		TObserverVariant(TVariant<TYPE_OF_NULLPTR, TWeakPtr<T>, TWeakObjectPtr<T>> ObserverVariant)
+		TObserverVariant(TVariant<TYPE_OF_NULLPTR, TWeakPtr<T>, TWeakInterfacePtr<T>> ObserverVariant)
 			: ObserverVariant(ObserverVariant)
 		{
 		}
@@ -27,7 +26,7 @@ namespace UE::PixelStreaming2
 		{
 			switch (ObserverVariant.GetIndex())
 			{
-				case TVariant<TYPE_OF_NULLPTR, TWeakPtr<T>, TWeakObjectPtr<T>>::template IndexOfType<TWeakPtr<T>>():
+				case TVariant<TYPE_OF_NULLPTR, TWeakPtr<T>, TWeakInterfacePtr<T>>::template IndexOfType<TWeakPtr<T>>():
 				{
 					if (TSharedPtr<T> PinnedUserObserver = ObserverVariant.template Get<TWeakPtr<T>>().Pin())
 					{
@@ -35,12 +34,12 @@ namespace UE::PixelStreaming2
 					}
 					break;
 				}
-				case TVariant<TYPE_OF_NULLPTR, TWeakPtr<T>, TWeakObjectPtr<T>>::template IndexOfType<TWeakObjectPtr<T>>():
+				case TVariant<TYPE_OF_NULLPTR, TWeakPtr<T>, TWeakInterfacePtr<T>>::template IndexOfType<TWeakInterfacePtr<T>>():
 				{
-					TWeakObjectPtr<T> WeakObserver = ObserverVariant.template Get<TWeakObjectPtr<T>>();
+					TWeakInterfacePtr<T> WeakObserver = ObserverVariant.template Get<TWeakInterfacePtr<T>>();
 					if (WeakObserver.IsValid())
 					{
-						return Cast<T>(WeakObserver.Get());
+						return WeakObserver.Get();
 					}
 					break;
 				}
@@ -55,13 +54,13 @@ namespace UE::PixelStreaming2
 		{
 			switch (ObserverVariant.GetIndex())
 			{
-				case TVariant<TYPE_OF_NULLPTR, TWeakPtr<T>, TWeakObjectPtr<T>>::template IndexOfType<TWeakPtr<T>>():
+				case TVariant<TYPE_OF_NULLPTR, TWeakPtr<T>, TWeakInterfacePtr<T>>::template IndexOfType<TWeakPtr<T>>():
 				{
 					return ObserverVariant.template Get<TWeakPtr<T>>().IsValid();
 				}
-				case TVariant<TYPE_OF_NULLPTR, TWeakPtr<T>, TWeakObjectPtr<T>>::template IndexOfType<TWeakObjectPtr<T>>():
+				case TVariant<TYPE_OF_NULLPTR, TWeakPtr<T>, TWeakInterfacePtr<T>>::template IndexOfType<TWeakInterfacePtr<T>>():
 				{
-					return ObserverVariant.template Get<TWeakObjectPtr<T>>().IsValid();
+					return ObserverVariant.template Get<TWeakInterfacePtr<T>>().IsValid();
 				}
 				default:
 					checkNoEntry();
@@ -70,18 +69,18 @@ namespace UE::PixelStreaming2
 		}
 
 	private:
-		TVariant<TYPE_OF_NULLPTR, TWeakPtr<T>, TWeakObjectPtr<T>> ObserverVariant;
+		TVariant<TYPE_OF_NULLPTR, TWeakPtr<T>, TWeakInterfacePtr<T>> ObserverVariant;
 	};
 
 	template <typename T>
 	inline TObserverVariant<T> TObserver(TWeakPtr<T> WeakObserver)
 	{
-		return TObserverVariant<T>(TVariant<TYPE_OF_NULLPTR, TWeakPtr<T>, TWeakObjectPtr<T>>(TInPlaceType<TWeakPtr<T>>(), WeakObserver));
+		return TObserverVariant<T>(TVariant<TYPE_OF_NULLPTR, TWeakPtr<T>, TWeakInterfacePtr<T>>(TInPlaceType<TWeakPtr<T>>(), WeakObserver));
 	}
 
 	template <typename T>
-	inline TObserverVariant<T> TObserver(TWeakObjectPtr<T> WeakObserver)
+	inline TObserverVariant<T> TObserver(TWeakInterfacePtr<T> WeakObserver)
 	{
-		return TObserverVariant<T>(TVariant<TYPE_OF_NULLPTR, TWeakPtr<T>, TWeakObjectPtr<T>>(TInPlaceType<TWeakObjectPtr<T>>(), WeakObserver));
+		return TObserverVariant<T>(TVariant<TYPE_OF_NULLPTR, TWeakPtr<T>, TWeakInterfacePtr<T>>(TInPlaceType<TWeakInterfacePtr<T>>(), WeakObserver));
 	}
 } // namespace UE::PixelStreaming2
\ No newline at end of file
diff --git a/Engine/Plugins/Media/PixelStreaming2/Source/PixelStreaming2RTC/Private/Blueprints/PixelStreaming2Peer.cpp b/Engine/Plugins/Media/PixelStreaming2/Source/PixelStreaming2RTC/Private/Blueprints/PixelStreaming2Peer.cpp
--- a/Engine/Plugins/Media/PixelStreaming2/Source/PixelStreaming2RTC/Private/Blueprints/PixelStreaming2Peer.cpp
+++ b/Engine/Plugins/Media/PixelStreaming2/Source/PixelStreaming2RTC/Private/Blueprints/PixelStreaming2Peer.cpp
@@ -27,17 +27,16 @@ UPixelStreaming2Peer::UPixelStreaming2Peer(const FObjectInitializer& ObjectIniti
 
 void UPixelStreaming2Peer::BeginPlay()
 {
-	const TWeakObjectPtr<UPixelStreaming2Peer> WeakThis = this;
-	SessionObserver = MakeRefCount<UE::PixelStreaming2::FEpicRtcSessionObserver>(UE::PixelStreaming2::TObserver<IPixelStreaming2SessionObserver>(WeakThis));
-	RoomObserver = MakeRefCount<UE::PixelStreaming2::FEpicRtcRoomObserver>(UE::PixelStreaming2::TObserver<IPixelStreaming2RoomObserver>(WeakThis));
+	Super::BeginPlay();
+	
+	SessionObserver = MakeRefCount<UE::PixelStreaming2::FEpicRtcSessionObserver>(UE::PixelStreaming2::TObserver<IPixelStreaming2SessionObserver>(this));
+	RoomObserver = MakeRefCount<UE::PixelStreaming2::FEpicRtcRoomObserver>(UE::PixelStreaming2::TObserver<IPixelStreaming2RoomObserver>(this));
 
-	AudioTrackObserverFactory = MakeRefCount<UE::PixelStreaming2::FEpicRtcAudioTrackObserverFactory>(UE::PixelStreaming2::TObserver<IPixelStreaming2AudioTrackObserver>(WeakThis));
-	VideoTrackObserverFactory = MakeRefCount<UE::PixelStreaming2::FEpicRtcVideoTrackObserverFactory>(UE::PixelStreaming2::TObserver<IPixelStreaming2VideoTrackObserver>(WeakThis));
-	DataTrackObserverFactory = MakeRefCount<UE::PixelStreaming2::FEpicRtcDataTrackObserverFactory>(UE::PixelStreaming2::TObserver<IPixelStreaming2DataTrackObserver>(WeakThis));
+	AudioTrackObserverFactory = MakeRefCount<UE::PixelStreaming2::FEpicRtcAudioTrackObserverFactory>(UE::PixelStreaming2::TObserver<IPixelStreaming2AudioTrackObserver>(this));
+	VideoTrackObserverFactory = MakeRefCount<UE::PixelStreaming2::FEpicRtcVideoTrackObserverFactory>(UE::PixelStreaming2::TObserver<IPixelStreaming2VideoTrackObserver>(this));
+	DataTrackObserverFactory = MakeRefCount<UE::PixelStreaming2::FEpicRtcDataTrackObserverFactory>(UE::PixelStreaming2::TObserver<IPixelStreaming2DataTrackObserver>(this));
 
 	UE::PixelStreaming2::FPixelStreaming2RTCModule::GetModule()->GetStatsCollector()->OnStatsReady.AddUObject(this, &UPixelStreaming2Peer::OnStatsReady);
-
-	Super::BeginPlay();
 }
 
 void UPixelStreaming2Peer::BeginDestroy()
