[PS2] Fix deadlock and crash when shutting down PixelStreaming2.

The problem is the PixelStreaming2 and PixelStreaming2RTC modules would be destroyed while a EpicRtcStreamer still lived. This would try to access the conference and rooms that no longer existed.

The fix was to store the conference in the streamer to keep it alive while a streamer lives and to shutdown all streams before destroying the modules.

diff --git a/Engine/Plugins/Media/PixelStreaming2/Source/PixelStreaming2/Private/PixelStreaming2Module.cpp b/Engine/Plugins/Media/PixelStreaming2/Source/PixelStreaming2/Private/PixelStreaming2Module.cpp
--- a/Engine/Plugins/Media/PixelStreaming2/Source/PixelStreaming2/Private/PixelStreaming2Module.cpp
+++ b/Engine/Plugins/Media/PixelStreaming2/Source/PixelStreaming2/Private/PixelStreaming2Module.cpp
@@ -89,6 +89,11 @@ namespace UE::PixelStreaming2
 				return;
 			}
 
+			// Make sure streaming is stopped before modules are unloaded.
+			FCoreDelegates::OnEnginePreExit.AddLambda([this]() {
+				StopStreaming();
+			});
+
 			const EVideoCodec SelectedCodec = UE::PixelStreaming2::GetEnumFromCVar<EVideoCodec>(UPixelStreaming2PluginSettings::CVarEncoderCodec);
 			if ((SelectedCodec == EVideoCodec::H264 && !IsHardwareEncoderSupported<FVideoEncoderConfigH264>())
 				|| (SelectedCodec == EVideoCodec::AV1 && !IsHardwareEncoderSupported<FVideoEncoderConfigAV1>()))
diff --git a/Engine/Plugins/Media/PixelStreaming2/Source/PixelStreaming2RTC/Private/EpicRtcStreamer.cpp b/Engine/Plugins/Media/PixelStreaming2/Source/PixelStreaming2RTC/Private/EpicRtcStreamer.cpp
--- a/Engine/Plugins/Media/PixelStreaming2/Source/PixelStreaming2RTC/Private/EpicRtcStreamer.cpp
+++ b/Engine/Plugins/Media/PixelStreaming2/Source/PixelStreaming2RTC/Private/EpicRtcStreamer.cpp
@@ -37,16 +37,16 @@
 
 namespace UE::PixelStreaming2
 {
-	FEpicRtcStreamer::FEpicRtcStreamer(const FString& InStreamerId)
+	FEpicRtcStreamer::FEpicRtcStreamer(const FString& InStreamerId, TRefCountPtr<EpicRtcConferenceInterface> Conference)
 		: StreamerId(InStreamerId)
 		, InputHandler(FRTCInputHandler::Create())
 		, Players(new TThreadSafeMap<FString, FPlayerContext>())
 		, VideoCapturer(FEpicRtcVideoCapturer::Create())
 		, VideoSourceGroup(FVideoSourceGroup::Create(VideoCapturer))
 		, FreezeFrame(FFreezeFrame::Create(Players, VideoCapturer, InputHandler))
+		, ReconnectTimer(MakeShared<FStreamerReconnectTimer>())
+		, EpicRtcConference(MoveTemp(Conference))
 	{
-		ReconnectTimer = MakeShared<FStreamerReconnectTimer>();
-
 		InputHandler->SetElevatedCheck([this](FString PlayerId) {
 			return GetEnumFromCVar<EInputControllerMode>(UPixelStreaming2PluginSettings::CVarInputController) == EInputControllerMode::Any
 				|| InputControllingId == INVALID_PLAYER_ID
@@ -61,12 +61,11 @@ namespace UE::PixelStreaming2
 		// Pin ref counted objects and clear the session if it is still alive.
 		// Session may still be alive if the manager was destroyed quickly after stopping the stream
 		// and the EpicRtcSessionState::Disconnected event did not have enough time to tick.
-		const TRefCountPtr<EpicRtcConferenceInterface> Conference = FPixelStreaming2RTCModule::GetModule()->GetEpicRtcConference();
-		const TRefCountPtr<EpicRtcSessionInterface>	   Session = EpicRtcSession;
-		if (Conference.IsValid() && Session.IsValid())
+		const TRefCountPtr<EpicRtcSessionInterface>	Session(EpicRtcSession);
+		if (EpicRtcConference.IsValid() && Session.IsValid())
 		{
 			const FUtf8String Utf8StreamerId(StreamerId);
-			Conference->RemoveSession(ToEpicRtcStringView(Utf8StreamerId));
+			EpicRtcConference->RemoveSession(ToEpicRtcStringView(Utf8StreamerId));
 			EpicRtcSession = nullptr;
 		}
 	}
@@ -244,6 +243,8 @@ namespace UE::PixelStreaming2
 			return;
 		}
 
+		check(EpicRtcConference.IsValid());
+
 		ReconnectTimer->Stop();
 
 		VideoCapturer->ResetFrameCapturer();
@@ -262,7 +263,7 @@ namespace UE::PixelStreaming2
 			._observer = SessionObserver
 		};
 
-		if (const EpicRtcErrorCode Result = FPixelStreaming2RTCModule::GetModule()->GetEpicRtcConference()->CreateSession(SessionConfig, EpicRtcSession.GetInitReference()); Result != EpicRtcErrorCode::Ok)
+		if (const EpicRtcErrorCode Result = EpicRtcConference->CreateSession(SessionConfig, EpicRtcSession.GetInitReference()); Result != EpicRtcErrorCode::Ok)
 		{
 			UE_LOGFMT(LogPixelStreaming2RTC, Error, "Failed to create EpicRtc session. CreateSession returned {0}", *ToString(Result));
 			StreamState = EStreamState::Disconnected;
@@ -1049,10 +1050,10 @@ namespace UE::PixelStreaming2
 
 				// The session has been disconnected (either through a call to StopStreaming or by an error)
 				// so remove it
-				if (FPixelStreaming2RTCModule::GetModule()->GetEpicRtcConference())
+				if (EpicRtcConference.IsValid())
 				{
 					FUtf8String Utf8StreamerId(StreamerId);
-					FPixelStreaming2RTCModule::GetModule()->GetEpicRtcConference()->RemoveSession(ToEpicRtcStringView(Utf8StreamerId));
+					EpicRtcConference->RemoveSession(ToEpicRtcStringView(Utf8StreamerId));
 					EpicRtcSession = nullptr;
 				}
 
@@ -1857,9 +1858,14 @@ namespace UE::PixelStreaming2
 		return RTC_STREAM_TYPE;
 	}
 
+	FRTCStreamerFactory::FRTCStreamerFactory(TRefCountPtr<EpicRtcConferenceInterface> Conference)
+		: EpicRtcConference(MoveTemp(Conference))
+	{
+	}
+
 	TSharedPtr<IPixelStreaming2Streamer> FRTCStreamerFactory::CreateNewStreamer(const FString& StreamerId)
 	{
-		TSharedPtr<IPixelStreaming2Streamer> NewStreamer = MakeShared<FEpicRtcStreamer>(StreamerId);
+		TSharedPtr<IPixelStreaming2Streamer> NewStreamer = MakeShared<FEpicRtcStreamer>(StreamerId, EpicRtcConference);
 
 		// default to the scene viewport if we have a game engine
 		if (UGameEngine* GameEngine = Cast<UGameEngine>(GEngine))
diff --git a/Engine/Plugins/Media/PixelStreaming2/Source/PixelStreaming2RTC/Private/EpicRtcStreamer.h b/Engine/Plugins/Media/PixelStreaming2/Source/PixelStreaming2RTC/Private/EpicRtcStreamer.h
--- a/Engine/Plugins/Media/PixelStreaming2/Source/PixelStreaming2RTC/Private/EpicRtcStreamer.h
+++ b/Engine/Plugins/Media/PixelStreaming2/Source/PixelStreaming2RTC/Private/EpicRtcStreamer.h
@@ -40,7 +40,7 @@ namespace UE::PixelStreaming2
 		public IPixelStreaming2VideoTrackObserver
 	{
 	public:
-		FEpicRtcStreamer(const FString& StreamerId);
+		FEpicRtcStreamer(const FString& StreamerId, TRefCountPtr<EpicRtcConferenceInterface> Conference);
 		virtual ~FEpicRtcStreamer();
 
 		virtual void  Initialize() override;
@@ -195,6 +195,7 @@ namespace UE::PixelStreaming2
 
 	private:
 		// Begin EpicRtc Classes
+		TRefCountPtr<EpicRtcConferenceInterface> EpicRtcConference;
 		TRefCountPtr<EpicRtcSessionInterface>	 EpicRtcSession;
 		TRefCountPtr<EpicRtcRoomInterface>		 EpicRtcRoom;
 		// End EpicRtc Classes
@@ -211,9 +212,12 @@ namespace UE::PixelStreaming2
 	class FRTCStreamerFactory : public IPixelStreaming2StreamerFactory
 	{
 	public:
-		FRTCStreamerFactory() = default;
+		FRTCStreamerFactory(TRefCountPtr<EpicRtcConferenceInterface> Conference);
+		virtual ~FRTCStreamerFactory() = default;
 
 		virtual FString								 GetStreamType() override;
 		virtual TSharedPtr<IPixelStreaming2Streamer> CreateNewStreamer(const FString& StreamerId) override;
+	private:
+		TRefCountPtr<EpicRtcConferenceInterface> EpicRtcConference;
 	};
 } // namespace UE::PixelStreaming2
diff --git a/Engine/Plugins/Media/PixelStreaming2/Source/PixelStreaming2RTC/Private/PixelStreaming2RTCModule.cpp b/Engine/Plugins/Media/PixelStreaming2/Source/PixelStreaming2RTC/Private/PixelStreaming2RTCModule.cpp
--- a/Engine/Plugins/Media/PixelStreaming2/Source/PixelStreaming2RTC/Private/PixelStreaming2RTCModule.cpp
+++ b/Engine/Plugins/Media/PixelStreaming2/Source/PixelStreaming2RTC/Private/PixelStreaming2RTCModule.cpp
@@ -123,7 +123,7 @@ namespace UE::PixelStreaming2
 				return;
 			}
 
-			StreamerFactory.Reset(new FRTCStreamerFactory());
+			StreamerFactory.Reset(new FRTCStreamerFactory(EpicRtcConference));
 
 			// Ensure we have ImageWrapper loaded, used in Freezeframes
 			verify(FModuleManager::Get().LoadModule(FName("ImageWrapper")));
