[PS2] Fix: Application crashing when peer disconnects with UseCamera

diff --git a/Engine/Plugins/Media/PixelStreaming2/Source/PixelStreaming2/Internal/ThreadSafeMap.h b/Engine/Plugins/Media/PixelStreaming2/Source/PixelStreaming2/Internal/ThreadSafeMap.h
--- a/Engine/Plugins/Media/PixelStreaming2/Source/PixelStreaming2/Internal/ThreadSafeMap.h
+++ b/Engine/Plugins/Media/PixelStreaming2/Source/PixelStreaming2/Internal/ThreadSafeMap.h
@@ -254,6 +254,13 @@ namespace UE::PixelStreaming2
 			return static_cast<Super*>(this)->Find(Key);
 		}
 
+		FORCEINLINE ValueType FindRef(KeyInitType Key)
+		{
+			FScopeLock Lock(&Mutex);
+
+			return static_cast<Super*>(this)->FindRef(Key);
+		}
+
 		FORCEINLINE int32 Remove(KeyConstPointerType Key)
 		{
 			FScopeLock Lock(&Mutex);
diff --git a/Engine/Plugins/Media/PixelStreaming2/Source/PixelStreaming2/Internal/VideoSink.h b/Engine/Plugins/Media/PixelStreaming2/Source/PixelStreaming2/Internal/VideoSink.h
--- a/Engine/Plugins/Media/PixelStreaming2/Source/PixelStreaming2/Internal/VideoSink.h
+++ b/Engine/Plugins/Media/PixelStreaming2/Source/PixelStreaming2/Internal/VideoSink.h
@@ -21,14 +21,13 @@ namespace UE::PixelStreaming2
 
 		void SetMuted(bool bIsMuted);
 
+		void OnVideoData(FTextureRHIRef Frame);
+
 	protected:
 		FCriticalSection					 VideoConsumersCS;
 		TSet<IPixelStreaming2VideoConsumer*> VideoConsumers;
 
 		bool bIsMuted = false;
-
-	protected:
-		void CallConsumeFrame(FTextureRHIRef Frame);
 	};
 
 } // namespace UE::PixelStreaming2
\ No newline at end of file
diff --git a/Engine/Plugins/Media/PixelStreaming2/Source/PixelStreaming2/Private/AudioSink.cpp b/Engine/Plugins/Media/PixelStreaming2/Source/PixelStreaming2/Private/AudioSink.cpp
--- a/Engine/Plugins/Media/PixelStreaming2/Source/PixelStreaming2/Private/AudioSink.cpp
+++ b/Engine/Plugins/Media/PixelStreaming2/Source/PixelStreaming2/Private/AudioSink.cpp
@@ -55,10 +55,11 @@ namespace UE::PixelStreaming2
 		// The sample rate and number of channels here has absolutely no relationship with PixelStreaming2AudioDeviceModule.
 		// The sample rate and number of channels here is determined adaptively by WebRTC's NetEQ class that selects sample rate/number of channels
 		// based on network conditions and other factors.
-		if (!HasAudioConsumers() || bIsMuted)
+		if (!HasAudioConsumers() || bIsMuted || IsEngineExitRequested())
 		{
 			return;
 		}
+		
 		// Iterate audio consumers and pass this data to their buffers
 		FScopeLock Lock(&AudioConsumersCS);
 		for (IPixelStreaming2AudioConsumer* AudioConsumer : AudioConsumers)
diff --git a/Engine/Plugins/Media/PixelStreaming2/Source/PixelStreaming2/Private/VideoSink.cpp b/Engine/Plugins/Media/PixelStreaming2/Source/PixelStreaming2/Private/VideoSink.cpp
--- a/Engine/Plugins/Media/PixelStreaming2/Source/PixelStreaming2/Private/VideoSink.cpp
+++ b/Engine/Plugins/Media/PixelStreaming2/Source/PixelStreaming2/Private/VideoSink.cpp
@@ -17,6 +17,10 @@ namespace UE::PixelStreaming2
 		{
 			IPixelStreaming2VideoConsumer* VideoConsumer = Iter.ElementIt->Value;
 			Iter.RemoveCurrent();
+			if (VideoConsumer != nullptr)
+			{
+				VideoConsumer->OnConsumerRemoved();
+			}
 		}
 	}
 
@@ -46,8 +50,15 @@ namespace UE::PixelStreaming2
 		return VideoConsumers.Num() > 0;
 	}
 
-	void FVideoSink::CallConsumeFrame(FTextureRHIRef Frame)
+	void FVideoSink::OnVideoData(FTextureRHIRef Frame)
 	{
+		TRACE_CPUPROFILER_EVENT_SCOPE_ON_CHANNEL_STR("FVideoSink::OnData", PixelStreaming2Channel);
+
+		if (!HasVideoConsumers() || bIsMuted || IsEngineExitRequested())
+		{
+			return;
+		}
+
 		// Iterate video consumers and pass this data to their buffers
 		FScopeLock Lock(&VideoConsumersCS);
 		for (IPixelStreaming2VideoConsumer* VideoConsumer : VideoConsumers)
diff --git a/Engine/Plugins/Media/PixelStreaming2/Source/PixelStreaming2RTC/Internal/EpicRtcVideoSink.h b/Engine/Plugins/Media/PixelStreaming2/Source/PixelStreaming2RTC/Internal/EpicRtcVideoSink.h
--- a/Engine/Plugins/Media/PixelStreaming2/Source/PixelStreaming2RTC/Internal/EpicRtcVideoSink.h
+++ b/Engine/Plugins/Media/PixelStreaming2/Source/PixelStreaming2RTC/Internal/EpicRtcVideoSink.h
@@ -21,7 +21,7 @@ namespace UE::PixelStreaming2
 		// Note: destructor will call destroy on any attached video consumers
 		virtual ~FEpicRtcVideoSink() = default;
 
-		void OnVideoData(const EpicRtcVideoFrame& Frame);
+		void OnEpicRtcFrame(const EpicRtcVideoFrame& Frame);
 
 	private:
 		FEpicRtcVideoSink(TRefCountPtr<EpicRtcVideoTrackInterface> InTrack);
diff --git a/Engine/Plugins/Media/PixelStreaming2/Source/PixelStreaming2RTC/Internal/UtilsString.h b/Engine/Plugins/Media/PixelStreaming2/Source/PixelStreaming2RTC/Internal/UtilsString.h
--- a/Engine/Plugins/Media/PixelStreaming2/Source/PixelStreaming2RTC/Internal/UtilsString.h
+++ b/Engine/Plugins/Media/PixelStreaming2/Source/PixelStreaming2RTC/Internal/UtilsString.h
@@ -144,6 +144,20 @@ namespace UE::PixelStreaming2
 		return TEXT("Unknown");
 	}
 
+	inline FString ToString(EpicRtcTrackState State)
+	{
+		switch (State)
+		{
+			case EpicRtcTrackState::New:
+				return TEXT("New");
+			case EpicRtcTrackState::Active:
+				return TEXT("Active");
+			case EpicRtcTrackState::Stopped:
+				return TEXT("Stopped");
+		}
+		return TEXT("Unknown");
+	}
+
 	inline FString ToString(const TSharedPtr<FJsonObject>& JsonObj, bool bPretty = true)
 	{
 		FString Res;
diff --git a/Engine/Plugins/Media/PixelStreaming2/Source/PixelStreaming2RTC/Private/Blueprints/PixelStreaming2Peer.cpp b/Engine/Plugins/Media/PixelStreaming2/Source/PixelStreaming2RTC/Private/Blueprints/PixelStreaming2Peer.cpp
--- a/Engine/Plugins/Media/PixelStreaming2/Source/PixelStreaming2RTC/Private/Blueprints/PixelStreaming2Peer.cpp
+++ b/Engine/Plugins/Media/PixelStreaming2/Source/PixelStreaming2RTC/Private/Blueprints/PixelStreaming2Peer.cpp
@@ -375,7 +375,7 @@ void UPixelStreaming2Peer::OnVideoTrackFrame(EpicRtcVideoTrackInterface* VideoTr
 		return;
 	}
 
-	VideoSink->OnVideoData(Frame);
+	VideoSink->OnEpicRtcFrame(Frame);
 }
 
 void UPixelStreaming2Peer::OnVideoTrackRemoved(EpicRtcVideoTrackInterface* VideoTrack)
diff --git a/Engine/Plugins/Media/PixelStreaming2/Source/PixelStreaming2RTC/Private/EpicRtcStreamer.cpp b/Engine/Plugins/Media/PixelStreaming2/Source/PixelStreaming2RTC/Private/EpicRtcStreamer.cpp
--- a/Engine/Plugins/Media/PixelStreaming2/Source/PixelStreaming2RTC/Private/EpicRtcStreamer.cpp
+++ b/Engine/Plugins/Media/PixelStreaming2/Source/PixelStreaming2RTC/Private/EpicRtcStreamer.cpp
@@ -40,10 +40,10 @@ namespace UE::PixelStreaming2
 	FEpicRtcStreamer::FEpicRtcStreamer(const FString& InStreamerId, TRefCountPtr<EpicRtcConferenceInterface> Conference)
 		: StreamerId(InStreamerId)
 		, InputHandler(FRTCInputHandler::Create())
-		, Players(new TThreadSafeMap<FString, FPlayerContext>())
+		, Participants(new TThreadSafeMap<FString, TSharedPtr<FPlayerContext>>())
 		, VideoCapturer(FEpicRtcVideoCapturer::Create())
 		, VideoSourceGroup(FVideoSourceGroup::Create(VideoCapturer))
-		, FreezeFrame(FFreezeFrame::Create(Players, VideoCapturer, InputHandler))
+		, FreezeFrame(FFreezeFrame::Create(Participants, VideoCapturer, InputHandler))
 		, ReconnectTimer(MakeShared<FStreamerReconnectTimer>())
 		, EpicRtcConference(MoveTemp(Conference))
 	{
@@ -61,7 +61,7 @@ namespace UE::PixelStreaming2
 		// Pin ref counted objects and clear the session if it is still alive.
 		// Session may still be alive if the manager was destroyed quickly after stopping the stream
 		// and the EpicRtcSessionState::Disconnected event did not have enough time to tick.
-		const TRefCountPtr<EpicRtcSessionInterface>	Session(EpicRtcSession);
+		const TRefCountPtr<EpicRtcSessionInterface> Session(EpicRtcSession);
 		if (EpicRtcConference.IsValid() && Session.IsValid())
 		{
 			const FUtf8String Utf8StreamerId(StreamerId);
@@ -188,7 +188,7 @@ namespace UE::PixelStreaming2
 
 	void FEpicRtcStreamer::OnProtocolUpdated()
 	{
-		Players->Apply([this](FString DataPlayerId, FPlayerContext& PlayerContext) {
+		Participants->Apply([this](FString DataPlayerId, TSharedPtr<FPlayerContext> PlayerContext) {
 			SendProtocol(DataPlayerId);
 		});
 	}
@@ -301,18 +301,16 @@ namespace UE::PixelStreaming2
 
 	void FEpicRtcStreamer::OnStatsReady(const FString& PlayerId, const EpicRtcConnectionStats& ConnectionStats)
 	{
-		FPlayerContext* PlayerContext = Players->Find(PlayerId);
-		if (!PlayerContext)
+		TSharedPtr<FRTCStatsCollector> StatsCollector;
+		if (TSharedPtr<FPlayerContext> Participant = Participants->FindRef(PlayerId); Participant.IsValid())
 		{
-			return;
+			StatsCollector = Participant->StatsCollector;
 		}
 
-		if (!PlayerContext->StatsCollector)
+		if (StatsCollector)
 		{
-			return;
+			StatsCollector->Process(ConnectionStats);
 		}
-
-		PlayerContext->StatsCollector->Process(ConnectionStats);
 	}
 
 	void FEpicRtcStreamer::OnFrameCapturerCreated()
@@ -432,10 +430,10 @@ namespace UE::PixelStreaming2
 
 	void FEpicRtcStreamer::SendAllPlayersMessage(FString MessageType, const FString& Descriptor)
 	{
-		Players->Apply([&MessageType, &Descriptor](FString PlayerId, FPlayerContext& PlayerContext) {
-			if (PlayerContext.DataTrack && !IsSFU(PlayerId))
+		Participants->Apply([&MessageType, &Descriptor](FString PlayerId, TSharedPtr<FPlayerContext> Participant) {
+			if (Participant->DataTrack && !IsSFU(PlayerId))
 			{
-				PlayerContext.DataTrack->SendMessage(MessageType, Descriptor);
+				Participant->DataTrack->SendMessage(MessageType, Descriptor);
 			}
 		});
 	}
@@ -446,13 +444,16 @@ namespace UE::PixelStreaming2
 		{
 			return;
 		}
-		if (FPlayerContext* PlayerContext = Players->Find(PlayerId))
+
+		TSharedPtr<FEpicRtcDataTrack> DataTrack;
+		if (TSharedPtr<FPlayerContext> Participant = Participants->FindRef(PlayerId); Participant.IsValid())
 		{
-			if (!PlayerContext->DataTrack)
-			{
-				return;
-			}
-			PlayerContext->DataTrack->SendMessage(MessageType, Descriptor);
+			DataTrack = Participant->DataTrack;
+		}
+
+		if (DataTrack)
+		{
+			DataTrack->SendMessage(MessageType, Descriptor);
 		}
 	}
 
@@ -462,61 +463,62 @@ namespace UE::PixelStreaming2
 		// channels it might be a bad idea. At some point it would be good to take a snapshot of the
 		// keys in the map when we start, then one by one get the channel and send the data
 
-		Players->Apply([&ByteData, &MimeType, &FileExtension](FString PlayerId, FPlayerContext& PlayerContext) {
-			if (!PlayerContext.DataTrack)
+		Participants->Apply([&ByteData, &MimeType, &FileExtension](FString PlayerId, TSharedPtr<FPlayerContext> Participant) {
+			if (!Participant->DataTrack)
 			{
 				return;
 			}
 
 			// Send the mime type first
-			PlayerContext.DataTrack->SendMessage(EPixelStreaming2FromStreamerMessage::FileMimeType, MimeType);
+			Participant->DataTrack->SendMessage(EPixelStreaming2FromStreamerMessage::FileMimeType, MimeType);
 
 			// Send the extension next
-			PlayerContext.DataTrack->SendMessage(EPixelStreaming2FromStreamerMessage::FileExtension, FileExtension);
+			Participant->DataTrack->SendMessage(EPixelStreaming2FromStreamerMessage::FileExtension, FileExtension);
 
 			// Send the contents of the file. Note to callers: consider running this on its own thread, it can take a while if the file is big.
-			PlayerContext.DataTrack->SendArbitraryData(EPixelStreaming2FromStreamerMessage::FileContents, ByteData);
+			Participant->DataTrack->SendArbitraryData(EPixelStreaming2FromStreamerMessage::FileContents, ByteData);
 		});
 	}
 
 	void FEpicRtcStreamer::KickPlayer(FString PlayerId)
 	{
-		if (FPlayerContext* PlayerContext = Players->Find(PlayerId))
+		TRefCountPtr<EpicRtcParticipantInterface> ParticipantInterface;
+		if (TSharedPtr<FPlayerContext> Participant = Participants->FindRef(PlayerId); Participant.IsValid())
 		{
-			PlayerContext->ParticipantInterface->Kick();
+			ParticipantInterface = Participant->ParticipantInterface;
 		}
+
+		ParticipantInterface->Kick();
 	}
 
 	TArray<FString> FEpicRtcStreamer::GetConnectedPlayers()
 	{
-		TArray<FString> ConnectedPlayerIds;
-		Players->Apply([&ConnectedPlayerIds, this](FString PlayerId, FPlayerContext& PlayerContext) {
-			ConnectedPlayerIds.Add(PlayerId);
-		});
-		return ConnectedPlayerIds;
+		TSet<FString> ConnectedParticipantIds;
+		Participants->GetKeys(ConnectedParticipantIds);
+
+		return ConnectedParticipantIds.Array();
 	}
 
 	IPixelStreaming2AudioSink* FEpicRtcStreamer::GetPeerAudioSink(FString PlayerId)
 	{
-		if (FPlayerContext* PlayerContext = Players->Find(PlayerId))
+		IPixelStreaming2AudioSink* Result = nullptr;
+		if (TSharedPtr<FPlayerContext> Participant = Participants->FindRef(PlayerId); Participant.IsValid())
 		{
-			if (PlayerContext->AudioSink)
-			{
-				return PlayerContext->AudioSink.Get();
-			}
+			Result = Participant->AudioSink ? Participant->AudioSink.Get() : nullptr;
 		}
-		return nullptr;
+
+		return Result;
 	}
 
 	IPixelStreaming2AudioSink* FEpicRtcStreamer::GetUnlistenedAudioSink()
 	{
 		IPixelStreaming2AudioSink* Result = nullptr;
-		Players->ApplyUntil([&Result](FString PlayerId, FPlayerContext& PlayerContext) {
-			if (PlayerContext.AudioSink)
+		Participants->ApplyUntil([&Result](FString PlayerId, TSharedPtr<FPlayerContext> Participant) {
+			if (Participant->AudioSink)
 			{
-				if (!PlayerContext.AudioSink->HasAudioConsumers())
+				if (!Participant->AudioSink->HasAudioConsumers())
 				{
-					Result = PlayerContext.AudioSink.Get();
+					Result = Participant->AudioSink.Get();
 					return true;
 				}
 			}
@@ -527,25 +529,24 @@ namespace UE::PixelStreaming2
 
 	IPixelStreaming2VideoSink* FEpicRtcStreamer::GetPeerVideoSink(FString PlayerId)
 	{
-		if (FPlayerContext* PlayerContext = Players->Find(PlayerId))
+		IPixelStreaming2VideoSink* Result = nullptr;
+		if (TSharedPtr<FPlayerContext> Participant = Participants->FindRef(PlayerId); Participant.IsValid())
 		{
-			if (PlayerContext->VideoSink)
-			{
-				return PlayerContext->VideoSink.Get();
-			}
+			Result = Participant->VideoSink ? Participant->VideoSink.Get() : nullptr;
 		}
-		return nullptr;
+
+		return Result;
 	}
 
 	IPixelStreaming2VideoSink* FEpicRtcStreamer::GetUnwatchedVideoSink()
 	{
 		IPixelStreaming2VideoSink* Result = nullptr;
-		Players->ApplyUntil([&Result](FString PlayerId, FPlayerContext& PlayerContext) {
-			if (PlayerContext.VideoSink)
+		Participants->ApplyUntil([&Result](FString PlayerId, TSharedPtr<FPlayerContext> Participant) {
+			if (Participant->VideoSink)
 			{
-				if (!PlayerContext.VideoSink->HasVideoConsumers())
+				if (!Participant->VideoSink->HasVideoConsumers())
 				{
-					Result = PlayerContext.VideoSink.Get();
+					Result = Participant->VideoSink.Get();
 					return true;
 				}
 			}
@@ -588,13 +589,13 @@ namespace UE::PixelStreaming2
 
 	void FEpicRtcStreamer::RefreshStreamBitrate()
 	{
-		Players->Apply([this](FString PlayerId, FPlayerContext& PlayerContext) {
-			if (!PlayerContext.ParticipantInterface)
+		Participants->Apply([this](FString PlayerId, TSharedPtr<FPlayerContext> Participant) {
+			if (!Participant->ParticipantInterface)
 			{
 				return;
 			}
 
-			TRefCountPtr<EpicRtcConnectionInterface> ConnectionInterface = PlayerContext.ParticipantInterface->GetConnection();
+			TRefCountPtr<EpicRtcConnectionInterface> ConnectionInterface = Participant->ParticipantInterface->GetConnection();
 			if (!ConnectionInterface)
 			{
 				return;
@@ -604,9 +605,9 @@ namespace UE::PixelStreaming2
 		});
 	}
 
-	void FEpicRtcStreamer::ForEachPlayer(const TFunction<void(FString, FPlayerContext)>& Func)
+	void FEpicRtcStreamer::ForEachPlayer(const TFunction<void(FString, TSharedPtr<FPlayerContext>)>& Func)
 	{
-		Players->Apply(Func);
+		Participants->Apply(Func);
 	}
 
 	void FEpicRtcStreamer::ConsumeStats(FString PlayerId, FName StatName, float StatValue)
@@ -621,14 +622,12 @@ namespace UE::PixelStreaming2
 			return;
 		}
 
-		const FPlayerContext* PlayerContext = Players->Find(PlayerId);
-
-		if (!PlayerContext)
+		TSharedPtr<FEpicRtcDataTrack> DataTrack;
+		if (TSharedPtr<FPlayerContext> Participant = Participants->FindRef(PlayerId); Participant.IsValid())
 		{
-			return;
+			DataTrack = Participant->DataTrack;
 		}
 
-		TSharedPtr<FEpicRtcDataTrack> DataTrack = PlayerContext->DataTrack;
 		if (!DataTrack)
 		{
 			return;
@@ -639,31 +638,29 @@ namespace UE::PixelStreaming2
 
 	void FEpicRtcStreamer::DeletePlayerSession(FString PlayerId)
 	{
-		// We dont want to allow this to be deleted within Players.Remove because
-		// we lock the players map and the delete could dispatch a webrtc object
-		// delete on the signalling thread which might be waiting for the players
-		// lock.
-		FPlayerContext PendingDeletePlayer;
-		if (FPlayerContext* PlayerContext = Players->Find(PlayerId))
+		// We need to keep a copy of the participant to delete as when destroying the ParticipantInterface,
+		// webrtc will finish tasks in the other threads which can execute FindRef and end up deadlocking
+		// the map. 
+		// TL:DR Participant deletion needs to happen outside the TThreadSafeMap scope lock
+		TSharedPtr<FPlayerContext> ParticipantToDelete;
+		if(!Participants->RemoveAndCopyValue(PlayerId, ParticipantToDelete) || !ParticipantToDelete.IsValid())
 		{
-			// when a sfu is connected we only get disconnect messages.
-			// we dont get connect messages but we might get datachannel requests which can result
-			// in players with no PeerConnection but a datachannel
-			if (PlayerContext->VideoSource)
-			{
-				VideoSourceGroup->RemoveVideoSource(PlayerContext->VideoSource.Get());
-			}
+			return;
+		}
 
-			// Close any data track related things (we do this here because RoomLeft happens before DataTrack stopped fires)
-			// So if we only did this in DataTrack stooped the playerId would already be removed
-			OnDataChannelClosed(PlayerId);
+		UE_LOGFMT(LogPixelStreaming2RTC, Verbose, "FEpicRtcStreamer::DeletePlayerSession(Participant [{0}])", PlayerId);
 
-			PendingDeletePlayer = *PlayerContext;
-		}
+		OnDataTrackClosed(PlayerId);
+		// FIXME (RTCP-7928): EpicRtc currently isn't broadcasting a stopped track state for remote track
+		OnAudioTrackClosed(PlayerId, true);
+		OnVideoTrackClosed(PlayerId, true);
 
-		Players->Remove(PlayerId);
+		if (ParticipantToDelete->VideoSource)
+		{
+			VideoSourceGroup->RemoveVideoSource(ParticipantToDelete->VideoSource.Get());
+		}
 
-		DoOnGameThread([StreamerId = this->StreamerId, PlayerId, bIsEmpty = Players->IsEmpty()]() {
+		DoOnGameThread([StreamerId = this->StreamerId, PlayerId, bIsEmpty = Participants->IsEmpty()]() {
 			if (UPixelStreaming2Delegates* Delegates = UPixelStreaming2Delegates::Get())
 			{
 				Delegates->OnClosedConnection.Broadcast(StreamerId, PlayerId);
@@ -689,24 +686,27 @@ namespace UE::PixelStreaming2
 			PSStats->RemoveAllPeerStats();
 		}
 
-		VideoSourceGroup->RemoveAllVideoSources();
+		TSet<FString> PlayerIds;
+		Participants->GetKeys(PlayerIds);
+		// We have to iterate the keys separately as OnDataTrackClosed also loops through the player map
+		// and we will deadlock
+		for (const FString& PlayerId : PlayerIds)
+		{
+			DeletePlayerSession(PlayerId);
+		}
 
+		// Further cleanup
+		VideoSourceGroup->RemoveAllVideoSources();
 		AudioTrackPlayerIdMap.Empty();
 		VideoTrackPlayerIdMap.Empty();
-		Players->Empty();
+		Participants->Empty();
 		InputControllingId = INVALID_PLAYER_ID;
-
-		DoOnGameThread([StreamerId = this->StreamerId]() {
-			if (UPixelStreaming2Delegates* Delegates = UPixelStreaming2Delegates::Get())
-			{
-				Delegates->OnAllConnectionsClosed.Broadcast(StreamerId);
-				Delegates->OnAllConnectionsClosedNative.Broadcast(StreamerId);
-			}
-		});
 	}
 
-	void FEpicRtcStreamer::OnDataChannelOpen(FString PlayerId)
+	void FEpicRtcStreamer::OnDataTrackOpen(FString PlayerId)
 	{
+		UE_LOGFMT(LogPixelStreaming2RTC, Verbose, "FEpicRtcStreamer::OnDataTrackOpen(Participant [{0}])", PlayerId);
+
 		// Only time we automatically make a new peer the input controlling host is if they are the first peer (and not the SFU).
 		const bool HostControlsInput = GetEnumFromCVar<EInputControllerMode>(UPixelStreaming2PluginSettings::CVarInputController) == EInputControllerMode::Host;
 		if (HostControlsInput && !IsSFU(PlayerId) && InputControllingId == INVALID_PLAYER_ID)
@@ -730,39 +730,109 @@ namespace UE::PixelStreaming2
 		SendPeerControllerMessages(PlayerId);
 	}
 
-	void FEpicRtcStreamer::OnDataChannelClosed(FString PlayerId)
+	void FEpicRtcStreamer::OnDataTrackClosed(FString PlayerId)
 	{
-		if (FPlayerContext* PlayerContext = Players->Find(PlayerId))
-		{
-			PlayerContext->DataTrack = nullptr;
+		UE_LOGFMT(LogPixelStreaming2RTC, Verbose, "FEpicRtcStreamer::OnDataTrackClosed(Participant [{0}])", PlayerId);
 
-			if (InputControllingId == PlayerId)
-			{
-				InputControllingId = INVALID_PLAYER_ID;
-				// just get the first channel we have and give it input control.
-				Players->ApplyUntil([this](FString PlayerId, FPlayerContext& PlayerContext) {
-					if (!PlayerContext.DataTrack)
-					{
-						return false;
-					}
-					if (IsSFU(PlayerId))
-					{
-						return false;
-					}
-					InputControllingId = PlayerId;
-					PlayerContext.DataTrack->SendMessage(EPixelStreaming2FromStreamerMessage::InputControlOwnership, 1 /* ControlsInput */);
-					return true;
-				});
-			}
+		if (TSharedPtr<FPlayerContext> Participant = Participants->FindRef(PlayerId); Participant.IsValid())
+		{
+			Participant->DataTrack = nullptr;
+		}
 
-			DoOnGameThread([StreamerId = this->StreamerId, PlayerId]() {
-				if (UPixelStreaming2Delegates* Delegates = UPixelStreaming2Delegates::Get())
+		if (InputControllingId == PlayerId)
+		{
+			InputControllingId = INVALID_PLAYER_ID;
+			// just get the first channel we have and give it input control.
+			Participants->ApplyUntil([this](FString PlayerId, TSharedPtr<FPlayerContext> Participant) {
+				if (!Participant->DataTrack)
 				{
-					Delegates->OnDataTrackClosed.Broadcast(StreamerId, PlayerId);
-					Delegates->OnDataTrackClosedNative.Broadcast(StreamerId, PlayerId);
+					return false;
 				}
+				if (IsSFU(PlayerId))
+				{
+					return false;
+				}
+				InputControllingId = PlayerId;
+				Participant->DataTrack->SendMessage(EPixelStreaming2FromStreamerMessage::InputControlOwnership, 1 /* ControlsInput */);
+				return true;
 			});
 		}
+
+		DoOnGameThread([StreamerId = this->StreamerId, PlayerId]() {
+			if (UPixelStreaming2Delegates* Delegates = UPixelStreaming2Delegates::Get())
+			{
+				Delegates->OnDataTrackClosed.Broadcast(StreamerId, PlayerId);
+				Delegates->OnDataTrackClosedNative.Broadcast(StreamerId, PlayerId);
+			}
+		});
+	}
+
+	void FEpicRtcStreamer::OnAudioTrackOpen(FString PlayerId, bool bIsRemote)
+	{
+		UE_LOGFMT(LogPixelStreaming2RTC, Verbose, "FEpicRtcStreamer::OnAudioTrackOpen(Participant [{0}], IsRemote [{1}])", PlayerId, bIsRemote);
+
+		if (UPixelStreaming2Delegates* Delegates = UPixelStreaming2Delegates::Get())
+		{
+			// NOTE: Native delegates aren't required to be broadcast on game thread
+			Delegates->OnAudioTrackOpenNative.Broadcast(StreamerId, PlayerId, bIsRemote);
+		}
+	}
+
+	void FEpicRtcStreamer::OnAudioTrackClosed(FString PlayerId, bool bIsRemote)
+	{
+		UE_LOGFMT(LogPixelStreaming2RTC, Verbose, "FEpicRtcStreamer::OnAudioTrackClosed(Participant [{0}], IsRemote [{1}])", PlayerId, bIsRemote);
+
+		if (TSharedPtr<FPlayerContext> Participant = Participants->FindRef(PlayerId); Participant.IsValid())
+		{
+			if (bIsRemote)
+			{
+				Participant->AudioSink = nullptr;
+			}
+			else
+			{
+				Participant->AudioSource = nullptr;
+			}
+		}
+
+		if (UPixelStreaming2Delegates* Delegates = UPixelStreaming2Delegates::Get())
+		{
+			// NOTE: Native delegates aren't required to be broadcast on game thread
+			Delegates->OnAudioTrackClosedNative.Broadcast(StreamerId, PlayerId, bIsRemote);
+		}
+	}
+
+	void FEpicRtcStreamer::OnVideoTrackOpen(FString PlayerId, bool bIsRemote)
+	{
+		UE_LOGFMT(LogPixelStreaming2RTC, Verbose, "FEpicRtcStreamer::OnVideoTrackOpen(Participant [{0}], IsRemote [{1}])", PlayerId, bIsRemote);
+
+		if (UPixelStreaming2Delegates* Delegates = UPixelStreaming2Delegates::Get())
+		{
+			// NOTE: Native delegates aren't required to be broadcast on game thread
+			Delegates->OnVideoTrackOpenNative.Broadcast(StreamerId, PlayerId, bIsRemote);
+		}
+	}
+
+	void FEpicRtcStreamer::OnVideoTrackClosed(FString PlayerId, bool bIsRemote)
+	{
+		UE_LOGFMT(LogPixelStreaming2RTC, Verbose, "FEpicRtcStreamer::OnVideoTrackClosed(Participant [{0}], IsRemote [{1}])", PlayerId, bIsRemote);
+
+		if (TSharedPtr<FPlayerContext> Participant = Participants->FindRef(PlayerId); Participant.IsValid())
+		{
+			if (bIsRemote)
+			{
+				Participant->VideoSink = nullptr;
+			}
+			else
+			{
+				Participant->VideoSource = nullptr;
+			}
+		}
+
+		if (UPixelStreaming2Delegates* Delegates = UPixelStreaming2Delegates::Get())
+		{
+			// NOTE: Native delegates aren't required to be broadcast on game thread
+			Delegates->OnVideoTrackClosedNative.Broadcast(StreamerId, PlayerId, bIsRemote);
+		}
 	}
 
 	void FEpicRtcStreamer::SendInitialSettings(FString PlayerId) const
@@ -796,14 +866,18 @@ namespace UE::PixelStreaming2
 
 		const FString FullPayload = FString::Printf(TEXT("{ \"PixelStreaming\": %s, \"Encoder\": %s, \"WebRTC\": %s, \"ConfigOptions\": %s }"), *PixelStreaming2Payload, *EncoderPayload, *WebRTCPayload, *ConfigPayload);
 
-		if (const FPlayerContext* PlayerContext = Players->Find(PlayerId))
+		TSharedPtr<FEpicRtcDataTrack> DataTrack;
+		if (TSharedPtr<FPlayerContext> Participant = Participants->FindRef(PlayerId); Participant.IsValid())
 		{
-			if (!PlayerContext->DataTrack)
-			{
-				return;
-			}
-			PlayerContext->DataTrack->SendMessage(EPixelStreaming2FromStreamerMessage::InitialSettings, FullPayload);
+			DataTrack = Participant->DataTrack;
+		}
+
+		if (!DataTrack)
+		{
+			return;
 		}
+
+		DataTrack->SendMessage(EPixelStreaming2FromStreamerMessage::InitialSettings, FullPayload);
 	}
 
 	void FEpicRtcStreamer::SendProtocol(FString PlayerId) const
@@ -820,30 +894,43 @@ namespace UE::PixelStreaming2
 				return;
 			}
 
-			if (const FPlayerContext* PlayerContext = Players->Find(PlayerId))
+			TSharedPtr<FEpicRtcDataTrack> DataTrack;
+			if (TSharedPtr<FPlayerContext> Participant = Participants->FindRef(PlayerId); Participant.IsValid())
 			{
-				if (!PlayerContext->DataTrack)
-				{
-					return;
-				}
-				PlayerContext->DataTrack->SendMessage(EPixelStreaming2FromStreamerMessage::Protocol, Body);
+				DataTrack = Participant->DataTrack;
 			}
+
+			if (!DataTrack)
+			{
+				return;
+			}
+
+			DataTrack->SendMessage(EPixelStreaming2FromStreamerMessage::Protocol, Body);
 		}
 	}
 
 	void FEpicRtcStreamer::SendPeerControllerMessages(FString PlayerId) const
 	{
-		if (const FPlayerContext* PlayerContext = Players->Find(PlayerId))
+		TSharedPtr<FEpicRtcDataTrack> DataTrack;
 		{
-			if (!PlayerContext->DataTrack)
+			TSharedPtr<FPlayerContext> Participant = Participants->FindRef(PlayerId);
+			if (!Participant)
 			{
 				return;
 			}
-			const uint8 ControlsInput = (GetEnumFromCVar<EInputControllerMode>(UPixelStreaming2PluginSettings::CVarInputController) == EInputControllerMode::Host) ? (PlayerId == InputControllingId) : 1;
-			// Even though the QualityController feature is removed we send it for backwards compatibility with older frontends (can probably remove 2 versions after 5.5)
-			PlayerContext->DataTrack->SendMessage(EPixelStreaming2FromStreamerMessage::InputControlOwnership, ControlsInput);
-			PlayerContext->DataTrack->SendMessage(EPixelStreaming2FromStreamerMessage::QualityControlOwnership, 1 /* True */);
+
+			DataTrack = Participant->DataTrack;
 		}
+
+		if (!DataTrack)
+		{
+			return;
+		}
+
+		const uint8 ControlsInput = (GetEnumFromCVar<EInputControllerMode>(UPixelStreaming2PluginSettings::CVarInputController) == EInputControllerMode::Host) ? (PlayerId == InputControllingId) : 1;
+		// Even though the QualityController feature is removed we send it for backwards compatibility with older frontends (can probably remove 2 versions after 5.5)
+		DataTrack->SendMessage(EPixelStreaming2FromStreamerMessage::InputControlOwnership, ControlsInput);
+		DataTrack->SendMessage(EPixelStreaming2FromStreamerMessage::QualityControlOwnership, 1 /* True */);
 	}
 
 	void FEpicRtcStreamer::SendLatencyReport(FString PlayerId) const
@@ -887,13 +974,18 @@ namespace UE::PixelStreaming2
 					TransmissionTimeMs);
 			}
 
-			if (const FPlayerContext* PlayerContext = Players->Find(PlayerId))
+			TSharedPtr<FEpicRtcDataTrack> DataTrack;
+			if (TSharedPtr<FPlayerContext> Participant = Participants->FindRef(PlayerId); Participant.IsValid())
 			{
-				if (PlayerContext->DataTrack)
-				{
-					PlayerContext->DataTrack->SendMessage(EPixelStreaming2FromStreamerMessage::LatencyTest, ReportToTransmitJSON);
-				}
+				DataTrack = Participant->DataTrack;
 			}
+
+			if (!DataTrack)
+			{
+				return;
+			}
+
+			DataTrack->SendMessage(EPixelStreaming2FromStreamerMessage::LatencyTest, ReportToTransmitJSON);
 		});
 	}
 
@@ -911,12 +1003,21 @@ namespace UE::PixelStreaming2
 			UE_LOG(LogPixelStreaming2RTC, Log, TEXT("FEpicRtcStreamer::HandleRelayStatusMessage() Adding new PlayerId [%s] with DataTrackId [%s]"), *PlayerId, *DataTrackId);
 
 			FString SFUId;
-			if (FindPlayerByDataTrack(DataTrack, SFUId))
+			if (FindPlayerFromTrack(DataTrack, SFUId))
 			{
-				FPlayerContext* SFUContext = Players->Find(SFUId);
-				FPlayerContext& PlayerContext = Players->FindOrAdd(PlayerId);
-				PlayerContext.DataTrack = FEpicRtcMutliplexDataTrack::Create(SFUContext->DataTrack, InputHandler->GetFromStreamerProtocol(), PlayerId);
-				OnDataChannelOpen(PlayerId);
+				TSharedPtr<FEpicRtcDataTrack> SFUDataTrack;
+				if (TSharedPtr<FPlayerContext> SFUParticipant = Participants->FindRef(SFUId); SFUParticipant.IsValid())
+				{
+					SFUDataTrack = SFUParticipant->DataTrack;
+				}
+
+				if (SFUDataTrack)
+				{
+					TSharedPtr<FPlayerContext>& Participant = Participants->FindOrAdd(PlayerId);
+					Participant = MakeShared<FPlayerContext>();
+					Participant->DataTrack = FEpicRtcMutliplexDataTrack::Create(SFUDataTrack, InputHandler->GetFromStreamerProtocol(), PlayerId);
+					OnDataTrackOpen(PlayerId);
+				}
 			}
 			else
 			{
@@ -927,8 +1028,8 @@ namespace UE::PixelStreaming2
 		{
 			UE_LOG(LogPixelStreaming2RTC, Log, TEXT("FEpicRtcStreamer::HandleRelayStatusMessage() Removing PlayerId [%s] with DataTrackId [%s]"), *PlayerId, *DataTrackId);
 
-			OnDataChannelClosed(PlayerId);
-			Players->Remove(PlayerId);
+			OnDataTrackClosed(PlayerId);
+			Participants->Remove(PlayerId);
 		}
 	}
 
@@ -1174,9 +1275,9 @@ namespace UE::PixelStreaming2
 		return bSyncVideoAndAudio ? "pixelstreaming_av_stream_id" : "pixelstreaming_video_stream_id";
 	}
 
-	void FEpicRtcStreamer::OnRoomJoinedUpdate(EpicRtcParticipantInterface* Participant)
+	void FEpicRtcStreamer::OnRoomJoinedUpdate(EpicRtcParticipantInterface* ParticipantInterface)
 	{
-		const FString ParticipantId = ToString(Participant->GetId());
+		const FString ParticipantId = ToString(ParticipantInterface->GetId());
 		UE_LOG(LogPixelStreaming2RTC, Log, TEXT("Player (%s) joined"), *ParticipantId);
 
 		if (ParticipantId == StreamerId)
@@ -1192,11 +1293,12 @@ namespace UE::PixelStreaming2
 			}
 		});
 
-		FPlayerContext& PlayerContext = Players->FindOrAdd(ParticipantId);
-		PlayerContext.ParticipantInterface = Participant;
-		PlayerContext.StatsCollector = FRTCStatsCollector::Create(ParticipantId);
+		TSharedPtr<FPlayerContext>& Participant = Participants->FindOrAdd(ParticipantId);
+		Participant = MakeShared<FPlayerContext>();
+		Participant->ParticipantInterface = ParticipantInterface;
+		Participant->StatsCollector = FRTCStatsCollector::Create(ParticipantId);
 
-		TRefCountPtr<EpicRtcConnectionInterface> ParticipantConnection = Participant->GetConnection();
+		TRefCountPtr<EpicRtcConnectionInterface> ParticipantConnection = ParticipantInterface->GetConnection();
 		ParticipantConnection->SetManualNegotiation(true);
 
 		const EVideoCodec SelectedCodec = GetEnumFromCVar<EVideoCodec>(UPixelStreaming2PluginSettings::CVarEncoderCodec);
@@ -1358,7 +1460,7 @@ namespace UE::PixelStreaming2
 		else
 		{
 			EpicRtcDataSource DataSource = {
-				._label = Participant->GetId(),
+				._label = ParticipantInterface->GetId(),
 				._maxRetransmitTime = 0,
 				._maxRetransmits = 0,
 				._isOrdered = true,
@@ -1381,71 +1483,68 @@ namespace UE::PixelStreaming2
 		DeletePlayerSession(ParticipantId);
 	}
 
-	void FEpicRtcStreamer::OnAudioTrackUpdate(EpicRtcParticipantInterface* Participant, EpicRtcAudioTrackInterface* AudioTrack)
+	void FEpicRtcStreamer::OnAudioTrackUpdate(EpicRtcParticipantInterface* ParticipantInterface, EpicRtcAudioTrackInterface* AudioTrack)
 	{
-		FString ParticipantId = ToString(Participant->GetId());
-		FString AudioTrackId = ToString(AudioTrack->GetId());
-		UE_LOG(LogPixelStreaming2RTC, Log, TEXT("FEpicRtcStreamer::OnAudioTrackUpdate(Participant [%s], AudioTrack [%s])"), *ParticipantId, *AudioTrackId);
+		const FString ParticipantId = ToString(ParticipantInterface->GetId());
+		const FString AudioTrackId = ToString(AudioTrack->GetId());
+		const bool	  bIsRemote = static_cast<bool>(AudioTrack->IsRemote());
 
-		if (FPlayerContext* PlayerContext = Players->Find(ParticipantId))
+		UE_LOGFMT(LogPixelStreaming2RTC, Log, "FEpicRtcStreamer::OnAudioTrackUpdate(Participant [{0}], AudioTrack [{1}, IsRemote [{2}]])", ParticipantId, AudioTrackId, bIsRemote);
+
+		AudioTrackPlayerIdMap.Add(reinterpret_cast<uintptr_t>(AudioTrack), ParticipantId);
+		if (TSharedPtr<FPlayerContext> Participant = Participants->FindRef(ParticipantId); Participant.IsValid())
 		{
-			AudioTrackPlayerIdMap.Add(reinterpret_cast<uintptr_t>(AudioTrack), ParticipantId);
-			if (AudioTrack->IsRemote())
+			if (bIsRemote)
 			{
-				PlayerContext->AudioSink = FEpicRtcAudioSink::Create(AudioTrack);
+				Participant->AudioSink = FEpicRtcAudioSink::Create(AudioTrack);
 			}
 			else
 			{
-				PlayerContext->AudioSource = FEpicRtcAudioSource::Create(AudioTrack, FPixelStreaming2RTCModule::GetModule()->GetAudioCapturer());
-			}
-
-			if (UPixelStreaming2Delegates* Delegates = UPixelStreaming2Delegates::Get())
-			{
-				Delegates->OnAudioTrackOpenNative.Broadcast(StreamerId, ParticipantId, (bool)AudioTrack->IsRemote());
+				Participant->AudioSource = FEpicRtcAudioSource::Create(AudioTrack, FPixelStreaming2RTCModule::GetModule()->GetAudioCapturer());
 			}
 		}
+
+		OnAudioTrackOpen(ParticipantId, bIsRemote);
 	}
 
-	void FEpicRtcStreamer::OnVideoTrackUpdate(EpicRtcParticipantInterface* Participant, EpicRtcVideoTrackInterface* VideoTrack)
+	void FEpicRtcStreamer::OnVideoTrackUpdate(EpicRtcParticipantInterface* ParticipantInterface, EpicRtcVideoTrackInterface* VideoTrack)
 	{
-		FString ParticipantId = ToString(Participant->GetId());
-		FString VideoTrackId = ToString(VideoTrack->GetId());
-		UE_LOG(LogPixelStreaming2RTC, Log, TEXT("FEpicRtcStreamer::OnVideoTrackUpdate(Participant [%s], VideoTrack [%s])"), *ParticipantId, *VideoTrackId);
+		const FString ParticipantId = ToString(ParticipantInterface->GetId());
+		const FString VideoTrackId = ToString(VideoTrack->GetId());
+		const bool	  bIsRemote = static_cast<bool>(VideoTrack->IsRemote());
 
-		if (FPlayerContext* PlayerContext = Players->Find(ParticipantId))
+		UE_LOGFMT(LogPixelStreaming2RTC, Log, "FEpicRtcStreamer::OnVideoTrackUpdate(Participant [{0}], VideoTrack [{1}], IsRemote[{2}])", ParticipantId, VideoTrackId, bIsRemote);
+
+		VideoTrackPlayerIdMap.Add(reinterpret_cast<uintptr_t>(VideoTrack), ParticipantId);
+		if (TSharedPtr<FPlayerContext> Participant = Participants->FindRef(ParticipantId); Participant.IsValid())
 		{
-			VideoTrackPlayerIdMap.Add(reinterpret_cast<uintptr_t>(VideoTrack), ParticipantId);
-			if (VideoTrack->IsRemote())
+			if (bIsRemote)
 			{
-				PlayerContext->VideoSink = FEpicRtcVideoSink::Create(VideoTrack);
+				Participant->VideoSink = FEpicRtcVideoSink::Create(VideoTrack);
 			}
 			else
 			{
-				PlayerContext->VideoSource = FEpicRtcVideoSource::Create(VideoTrack, VideoCapturer, VideoSourceGroup);
-			}
-
-			if (UPixelStreaming2Delegates* Delegates = UPixelStreaming2Delegates::Get())
-			{
-				Delegates->OnVideoTrackOpenNative.Broadcast(StreamerId, ParticipantId, (bool)VideoTrack->IsRemote());
+				Participant->VideoSource = FEpicRtcVideoSource::Create(VideoTrack, VideoCapturer, VideoSourceGroup);
 			}
 		}
+		OnVideoTrackOpen(ParticipantId, bIsRemote);
 	}
 
-	void FEpicRtcStreamer::OnDataTrackUpdate(EpicRtcParticipantInterface* Participant, EpicRtcDataTrackInterface* DataTrack)
+	void FEpicRtcStreamer::OnDataTrackUpdate(EpicRtcParticipantInterface* ParticipantInterface, EpicRtcDataTrackInterface* DataTrack)
 	{
-		FString ParticipantId = ToString(Participant->GetId());
+		FString ParticipantId = ToString(ParticipantInterface->GetId());
 		FString DataTrackId = ToString(DataTrack->GetId());
-		UE_LOG(LogPixelStreaming2RTC, Log, TEXT("FEpicRtcStreamer::OnDataTrackUpdate(Participant [%s], DataTrack [%s])"), *ParticipantId, *DataTrackId);
+		UE_LOGFMT(LogPixelStreaming2RTC, Log, "FEpicRtcStreamer::OnDataTrackUpdate(Participant [{0}], DataTrack [{1}])", ParticipantId, DataTrackId);
 
-		if (FPlayerContext* PlayerContext = Players->Find(ParticipantId))
+		if (TSharedPtr<FPlayerContext> Participant = Participants->FindRef(ParticipantId); Participant.IsValid())
 		{
-			if (!PlayerContext->DataTrack)
+			if (!Participant->DataTrack)
 			{
-				PlayerContext->DataTrack = FEpicRtcDataTrack::Create(DataTrack, InputHandler->GetFromStreamerProtocol());
+				Participant->DataTrack = FEpicRtcDataTrack::Create(DataTrack, InputHandler->GetFromStreamerProtocol());
 			}
 			else
 			{
-				PlayerContext->DataTrack->SetSendTrack(DataTrack);
+				Participant->DataTrack->SetSendTrack(DataTrack);
 			}
 		}
 	}
@@ -1496,7 +1595,7 @@ namespace UE::PixelStreaming2
 	void FEpicRtcStreamer::OnAudioTrackMuted(EpicRtcAudioTrackInterface* AudioTrack, EpicRtcBool bIsMuted)
 	{
 		FString PlayerId;
-		bool	bFoundPlayer = FindPlayerByAudioTrack(AudioTrack, PlayerId);
+		bool	bFoundPlayer = FindPlayerFromTrack(AudioTrack, PlayerId);
 		FString AudioTrackId = ToString(AudioTrack->GetId());
 		UE_LOG(LogPixelStreaming2RTC, Log, TEXT("FEpicRtcStreamer::OnAudioTrackMuted(AudioTrack [%s], bIsMuted[%s], PlayerId[%s])"), *AudioTrackId, bIsMuted ? TEXT("True") : TEXT("False"), *PlayerId);
 
@@ -1506,21 +1605,38 @@ namespace UE::PixelStreaming2
 			return;
 		}
 
-		FPlayerContext* PlayerContext = Players->Find(PlayerId);
 		if (AudioTrack->IsRemote())
 		{
-			PlayerContext->AudioSink->SetMuted((bool)bIsMuted);
+			TSharedPtr<FEpicRtcAudioSink> AudioSink;
+			if (TSharedPtr<FPlayerContext> Participant = Participants->FindRef(PlayerId); Participant.IsValid())
+			{
+				AudioSink = Participant->AudioSink;
+			}
+
+			if (AudioSink)
+			{
+				AudioSink->SetMuted(static_cast<bool>(bIsMuted));
+			}
 		}
 		else
 		{
-			PlayerContext->AudioSource->SetMuted((bool)bIsMuted);
+			TSharedPtr<FEpicRtcAudioSource> AudioSource;
+			if (TSharedPtr<FPlayerContext> Participant = Participants->FindRef(PlayerId); Participant.IsValid())
+			{
+				AudioSource = Participant->AudioSource;
+			}
+
+			if (AudioSource)
+			{
+				AudioSource->SetMuted(static_cast<bool>(bIsMuted));
+			}
 		}
 	}
 
 	void FEpicRtcStreamer::OnAudioTrackFrame(EpicRtcAudioTrackInterface* AudioTrack, const EpicRtcAudioFrame& Frame)
 	{
 		FString PlayerId;
-		bool	bFoundPlayer = FindPlayerByAudioTrack(AudioTrack, PlayerId);
+		bool	bFoundPlayer = FindPlayerFromTrack(AudioTrack, PlayerId);
 		FString AudioTrackId = ToString(AudioTrack->GetId());
 
 		if (!bFoundPlayer)
@@ -1529,10 +1645,15 @@ namespace UE::PixelStreaming2
 			return;
 		}
 
-		FPlayerContext* PlayerContext = Players->Find(PlayerId);
-		if (PlayerContext->AudioSink)
+		TSharedPtr<FEpicRtcAudioSink> AudioSink;
+		if (TSharedPtr<FPlayerContext> Participant = Participants->FindRef(PlayerId); Participant.IsValid())
+		{
+			AudioSink = Participant->AudioSink;
+		}
+
+		if (AudioSink)
 		{
-			PlayerContext->AudioSink->OnAudioData(Frame._data, Frame._length, Frame._format._numChannels, Frame._format._sampleRate);
+			AudioSink->OnAudioData(Frame._data, Frame._length, Frame._format._numChannels, Frame._format._sampleRate);
 		}
 	}
 
@@ -1542,60 +1663,32 @@ namespace UE::PixelStreaming2
 		UE_LOGFMT(LogPixelStreaming2RTC, Warning, "FEpicRtcStreamer::OnAudioTrackRemoved was called. Use OnAudioTrackState with State [EpicRtcTrackState::Stopped]");
 	}
 
-	void FEpicRtcStreamer::RemoveAudioTrack(EpicRtcAudioTrackInterface* AudioTrack)
+	void FEpicRtcStreamer::OnAudioTrackState(EpicRtcAudioTrackInterface* AudioTrack, const EpicRtcTrackState State)
 	{
-		FString PlayerId;
-		bool	bFoundPlayer = FindPlayerByAudioTrack(AudioTrack, PlayerId);
-		FString AudioTrackId = ToString(AudioTrack->GetId());
-		UE_LOG(LogPixelStreaming2RTC, Log, TEXT("FEpicRtcStreamer::OnAudioTrackRemoved(AudioTrack [%s], PlayerId[%s])"), *AudioTrackId, *PlayerId);
+		FString		  PlayerId;
+		const bool	  bFoundPlayer = FindPlayerFromTrack(AudioTrack, PlayerId);
+		const bool	  bIsRemote = static_cast<bool>(AudioTrack->IsRemote());
+		const FString AudioTrackId = ToString(AudioTrack->GetId());
 
 		if (!bFoundPlayer)
 		{
-			UE_LOG(LogPixelStreaming2RTC, Warning, TEXT("FEpicRtcStreamer::OnAudioTrackFrame(Failed to find a player for audio track [%s])"), *AudioTrackId);
+			// Not finding a player is expected as OnAudioTrackState will happen for OnAudioTrackUpdate during track addition
+			UE_LOGFMT(LogPixelStreaming2RTC, Verbose, "FEpicRtcStreamer::OnAudioTrackState(Cannot to find a player for audio track [{0}])", AudioTrackId);
 			return;
 		}
 
-		AudioTrackPlayerIdMap.Remove(reinterpret_cast<uintptr_t>(AudioTrack));
-		FPlayerContext* PlayerContext = Players->Find(PlayerId);
-		if (PlayerContext)
-		{
-			if (AudioTrack->IsRemote())
-			{
-				PlayerContext->AudioSink = nullptr;
-			}
-			else
-			{
-				PlayerContext->AudioSource = nullptr;
-			}
-		}
-
-		if (UPixelStreaming2Delegates* Delegates = UPixelStreaming2Delegates::Get())
+		UE_LOGFMT(LogPixelStreaming2RTC, Log, "FEpicRtcStreamer::OnAudioTrackState(AudioTrack=[{0}], Player=[{1}], State=[{2}])", AudioTrackId, PlayerId, ToString(State));
+		if (State == EpicRtcTrackState::Stopped)
 		{
-			Delegates->OnAudioTrackClosedNative.Broadcast(StreamerId, PlayerId, (bool)AudioTrack->IsRemote());
-		}
-	}
-
-	void FEpicRtcStreamer::OnAudioTrackState(EpicRtcAudioTrackInterface* AudioTrack, const EpicRtcTrackState State)
-	{
-		FString PlayerId;
-		FindPlayerByAudioTrack(AudioTrack, PlayerId);
-		const FString AudioTrackId = ToString(AudioTrack->GetId());
-
-		if (State == EpicRtcTrackState::Active)
-		{
-			UE_LOGFMT(LogPixelStreaming2RTC, Log, "FEpicRtcStreamer::OnAudioTrackState(AudioTrack=[{0}], Player=[{1}], State=Active)", *AudioTrackId, *PlayerId);
-		}
-		else if (State == EpicRtcTrackState::Stopped)
-		{
-			UE_LOGFMT(LogPixelStreaming2RTC, Log, "FEpicRtcStreamer::OnAudioTrackState(AudioTrack=[{0}], Player=[{1}], State=Stopped)", *AudioTrackId, *PlayerId);
-			RemoveAudioTrack(AudioTrack);
+			AudioTrackPlayerIdMap.Remove(reinterpret_cast<uintptr_t>(AudioTrack));
+			OnAudioTrackClosed(PlayerId, bIsRemote);
 		}
 	}
 
 	void FEpicRtcStreamer::OnVideoTrackMuted(EpicRtcVideoTrackInterface* VideoTrack, EpicRtcBool bIsMuted)
 	{
 		FString PlayerId;
-		bool	bFoundPlayer = FindPlayerByVideoTrack(VideoTrack, PlayerId);
+		bool	bFoundPlayer = FindPlayerFromTrack(VideoTrack, PlayerId);
 		FString VideoTrackId = ToString(VideoTrack->GetId());
 		UE_LOG(LogPixelStreaming2RTC, Log, TEXT("FEpicRtcStreamer::OnVideoTrackMuted(VideoTrack [%s], bIsMuted[%s], PlayerId[%s])"), *VideoTrackId, bIsMuted ? TEXT("True") : TEXT("False"), *PlayerId);
 
@@ -1605,21 +1698,38 @@ namespace UE::PixelStreaming2
 			return;
 		}
 
-		FPlayerContext* PlayerContext = Players->Find(PlayerId);
 		if (VideoTrack->IsRemote())
 		{
-			PlayerContext->VideoSink->SetMuted((bool)bIsMuted);
+			TSharedPtr<FEpicRtcVideoSink> VideoSink;
+			if (TSharedPtr<FPlayerContext> Participant = Participants->FindRef(PlayerId); Participant.IsValid())
+			{
+				VideoSink = Participant->VideoSink;
+			}
+
+			if (VideoSink)
+			{
+				VideoSink->SetMuted(static_cast<bool>(bIsMuted));
+			}
 		}
 		else
 		{
-			PlayerContext->VideoSource->SetMuted((bool)bIsMuted);
+			TSharedPtr<FEpicRtcVideoSource> VideoSource;
+			if (TSharedPtr<FPlayerContext> Participant = Participants->FindRef(PlayerId); Participant.IsValid())
+			{
+				VideoSource = Participant->VideoSource;
+			}
+
+			if (VideoSource)
+			{
+				VideoSource->SetMuted(static_cast<bool>(bIsMuted));
+			}
 		}
 	}
 
 	void FEpicRtcStreamer::OnVideoTrackFrame(EpicRtcVideoTrackInterface* VideoTrack, const EpicRtcVideoFrame& Frame)
 	{
 		FString PlayerId;
-		bool	bFoundPlayer = FindPlayerByVideoTrack(VideoTrack, PlayerId);
+		bool	bFoundPlayer = FindPlayerFromTrack(VideoTrack, PlayerId);
 		FString VideoTrackId = ToString(VideoTrack->GetId());
 
 		if (!bFoundPlayer)
@@ -1628,70 +1738,43 @@ namespace UE::PixelStreaming2
 			return;
 		}
 
-		FPlayerContext* PlayerContext = Players->Find(PlayerId);
-		if (PlayerContext->VideoSink)
+		TSharedPtr<FEpicRtcVideoSink> VideoSink;
+		if (TSharedPtr<FPlayerContext> Participant = Participants->FindRef(PlayerId); Participant.IsValid())
 		{
-			PlayerContext->VideoSink->OnVideoData(Frame);
+			VideoSink = Participant->VideoSink;
+		}
+
+		if (VideoSink)
+		{
+			VideoSink->OnEpicRtcFrame(Frame);
 		}
 	}
 
 	void FEpicRtcStreamer::OnVideoTrackRemoved(EpicRtcVideoTrackInterface* VideoTrack)
 	{
 		// TODO (Eden.Harris) OnVideoTrackRemoved is not hooked up in EpicRtc.
-		UE_LOGFMT(LogPixelStreaming2RTC, Warning, "FEpicRtcStreamer::OnVideoTrackRemoved was called. Use OnAudioTrackState with State [EpicRtcTrackState::Stopped]");
+		UE_LOGFMT(LogPixelStreaming2RTC, Warning, "FEpicRtcStreamer::OnVideoTrackRemoved was called. Use OnVideoTrackState with State [EpicRtcTrackState::Stopped]");
 	}
 
-	void FEpicRtcStreamer::RemoveVideoTrack(EpicRtcVideoTrackInterface* VideoTrack)
+	void FEpicRtcStreamer::OnVideoTrackState(EpicRtcVideoTrackInterface* VideoTrack, const EpicRtcTrackState State)
 	{
-		FString PlayerId;
-		bool	bFoundPlayer = FindPlayerByVideoTrack(VideoTrack, PlayerId);
-		FString VideoTrackId = ToString(VideoTrack->GetId());
+		FString		  PlayerId;
+		const bool	  bFoundPlayer = FindPlayerFromTrack(VideoTrack, PlayerId);
+		const bool	  bIsRemote = static_cast<bool>(VideoTrack->IsRemote());
+		const FString VideoTrackId = ToString(VideoTrack->GetId());
 
 		if (!bFoundPlayer)
 		{
-			UE_LOG(LogPixelStreaming2RTC, Warning, TEXT("FEpicRtcStreamer::OnVideoTrackRemoved(Failed to find a player for video track [%s])"), *VideoTrackId);
+			// Not finding a player is expected as OnVideoTrackState will happen for OnVideoTrackUpdate during track addition
+			UE_LOGFMT(LogPixelStreaming2RTC, Verbose, "FEpicRtcStreamer::OnVideoTrackState(Cannot to find a player for video track [{0}])", VideoTrackId);
 			return;
 		}
 
-		UE_LOG(LogPixelStreaming2RTC, Log, TEXT("FEpicRtcStreamer::OnVideoTrackRemoved(VideoTrack=[%s], Player=[%s])"), *VideoTrackId, *PlayerId);
-
-		VideoTrackPlayerIdMap.Remove(reinterpret_cast<uintptr_t>(VideoTrack));
-		// If we did find a player, clear its video sink/source
-		FPlayerContext* PlayerContext = Players->Find(PlayerId);
-		if (PlayerContext)
-		{
-			if (VideoTrack->IsRemote())
-			{
-				PlayerContext->VideoSink = nullptr;
-			}
-			else
-			{
-				PlayerContext->VideoSource = nullptr;
-			}
-		}
-
-		if (UPixelStreaming2Delegates* Delegates = UPixelStreaming2Delegates::Get())
+		UE_LOGFMT(LogPixelStreaming2RTC, Log, "FEpicRtcStreamer::OnVideoTrackState(VideoTrack=[{0}], Player=[{1}], State=[{2}])", VideoTrackId, PlayerId, ToString(State));
+		if (State == EpicRtcTrackState::Stopped)
 		{
-			Delegates->OnVideoTrackClosedNative.Broadcast(StreamerId, PlayerId, (bool)VideoTrack->IsRemote());
-		}
-	}
-
-	void FEpicRtcStreamer::OnVideoTrackState(EpicRtcVideoTrackInterface* VideoTrack, const EpicRtcTrackState State)
-	{
-		FString PlayerId;
-		FindPlayerByVideoTrack(VideoTrack, PlayerId);
-		const FString VideoTrackId = ToString(VideoTrack->GetId());
-
-		// Note: It is acceptable to not have a found a player for track state changes, as these can trigger before we have added a participant
-
-		if (State == EpicRtcTrackState::Active)
-		{
-			UE_LOGFMT(LogPixelStreaming2RTC, Log, "FEpicRtcStreamer::OnVideoTrackState(VideoTrack=[{0}], Player=[{1}], State=Active)", *VideoTrackId, *PlayerId);
-		}
-		else if (State == EpicRtcTrackState::Stopped)
-		{
-			UE_LOGFMT(LogPixelStreaming2RTC, Log, "FEpicRtcStreamer::OnVideoTrackState(VideoTrack=[{0}], Player=[{1}], State=Stopped)", *VideoTrackId, *PlayerId);
-			RemoveVideoTrack(VideoTrack);
+			VideoTrackPlayerIdMap.Remove(reinterpret_cast<uintptr_t>(VideoTrack));
+			OnVideoTrackClosed(PlayerId, bIsRemote);
 		}
 	}
 
@@ -1707,23 +1790,23 @@ namespace UE::PixelStreaming2
 	void FEpicRtcStreamer::OnDataTrackState(EpicRtcDataTrackInterface* DataTrack, const EpicRtcTrackState State)
 	{
 		FString PlayerId;
-		bool	bFoundPlayer = FindPlayerByDataTrack(DataTrack, PlayerId);
+		bool	bFoundPlayer = FindPlayerFromTrack(DataTrack, PlayerId);
 		FString DataTrackId = ToString(DataTrack->GetId());
 
+		if (!bFoundPlayer)
+		{
+			UE_LOGFMT(LogPixelStreaming2RTC, Warning, "FEpicRtcStreamer::OnDataTrackState(Failed to find a player for data track [{0}])", DataTrackId);
+			return;
+		}
+
+		UE_LOGFMT(LogPixelStreaming2RTC, Log, "FEpicRtcStreamer::OnDataTrackState(DataTrack=[{0}], Player=[{1}], State=[{2}])", DataTrackId, PlayerId, ToString(State));
 		if (State == EpicRtcTrackState::Active)
 		{
-			if (!bFoundPlayer)
-			{
-				UE_LOG(LogPixelStreaming2RTC, Log, TEXT("FEpicRtcStreamer::OnDataTrackState(Message was State=Active. Failed to find a player for data track [%s])"), *DataTrackId);
-				return;
-			}
-			UE_LOG(LogPixelStreaming2RTC, Log, TEXT("FEpicRtcStreamer::OnDataTrackState(Message was State=Active. Player [%s])"), *DataTrackId);
-			OnDataChannelOpen(PlayerId);
+			OnDataTrackOpen(PlayerId);
 		}
 		else if (State == EpicRtcTrackState::Stopped)
 		{
-			UE_LOG(LogPixelStreaming2RTC, Log, TEXT("FEpicRtcStreamer::OnDataTrackState(Message was State=Stopped. Player [%s])"), *DataTrackId);
-			OnDataChannelClosed(PlayerId);
+			OnDataTrackClosed(PlayerId);
 		}
 	}
 
@@ -1755,7 +1838,7 @@ namespace UE::PixelStreaming2
 			HandleRelayStatusMessage(Data, DataSize, DataTrack);
 			return;
 		}
-		else if (!FindPlayerByDataTrack(DataTrack, PlayerId))
+		else if (!FindPlayerFromTrack(DataTrack, PlayerId))
 		{
 			UE_LOG(LogPixelStreaming2RTC, Log, TEXT("FEpicRtcStreamer::OnDataTrackMessage(Failed to find a player for data track [%s])"), *DataTrackId);
 			return;
@@ -1775,15 +1858,18 @@ namespace UE::PixelStreaming2
 		}
 		else if (Type == ToStreamerProtocol->Find(EPixelStreaming2ToStreamerMessage::TestEcho)->GetID())
 		{
-			if (FPlayerContext* PlayerContext = Players->Find(PlayerId))
+			TSharedPtr<FEpicRtcDataTrack> ParticipantDataTrack;
+			if (TSharedPtr<FPlayerContext> Participant = Participants->FindRef(PlayerId); Participant.IsValid())
 			{
-				if (PlayerContext->DataTrack)
-				{
-					const size_t  DescriptorSize = (DataSize - 1) / sizeof(TCHAR);
-					const TCHAR*  DescPtr = reinterpret_cast<const TCHAR*>(Data + 1);
-					const FString Message(DescriptorSize, DescPtr);
-					PlayerContext->DataTrack->SendMessage(EPixelStreaming2FromStreamerMessage::TestEcho, Message);
-				}
+				ParticipantDataTrack = Participant->DataTrack;
+			}
+
+			if (ParticipantDataTrack)
+			{
+				const size_t  DescriptorSize = (DataSize - 1) / sizeof(TCHAR);
+				const TCHAR*  DescPtr = reinterpret_cast<const TCHAR*>(Data + 1);
+				const FString Message(DescriptorSize, DescPtr);
+				ParticipantDataTrack->SendMessage(EPixelStreaming2FromStreamerMessage::TestEcho, Message);
 			}
 		}
 		else if (!IsEngineExitRequested())
@@ -1806,53 +1892,6 @@ namespace UE::PixelStreaming2
 	{
 	}
 
-	bool FEpicRtcStreamer::FindPlayerByAudioTrack(EpicRtcAudioTrackInterface* AudioTrack, FString& OutPlayerId)
-	{
-		OutPlayerId = "";
-
-		FString* FoundPlayerId = AudioTrackPlayerIdMap.Find(reinterpret_cast<uintptr_t>(AudioTrack));
-		if (FoundPlayerId)
-		{
-			OutPlayerId = *FoundPlayerId;
-		}
-
-		return !OutPlayerId.IsEmpty();
-	}
-
-	bool FEpicRtcStreamer::FindPlayerByVideoTrack(EpicRtcVideoTrackInterface* VideoTrack, FString& OutPlayerId)
-	{
-		OutPlayerId = "";
-
-		FString* FoundPlayerId = VideoTrackPlayerIdMap.Find(reinterpret_cast<uintptr_t>(VideoTrack));
-		if (FoundPlayerId)
-		{
-			OutPlayerId = *FoundPlayerId;
-		}
-
-		return !OutPlayerId.IsEmpty();
-	}
-
-	bool FEpicRtcStreamer::FindPlayerByDataTrack(EpicRtcDataTrackInterface* DataTrack, FString& OutPlayerId)
-	{
-		OutPlayerId = "";
-		// Find the player corresponding to this data track
-		FString DataTrackId = ToString(DataTrack->GetId());
-		Players->ApplyUntil([this, DataTrackId, &OutPlayerId](FString DataPlayerId, FPlayerContext& PlayerContext) {
-			if (PlayerContext.DataTrack)
-			{
-				FString TrackId = ToString(PlayerContext.DataTrack->GetId());
-				if (TrackId == DataTrackId)
-				{
-					OutPlayerId = DataPlayerId;
-					return true;
-				}
-			}
-			return false;
-		});
-
-		return !OutPlayerId.IsEmpty();
-	}
-
 	FString FRTCStreamerFactory::GetStreamType()
 	{
 		return RTC_STREAM_TYPE;
diff --git a/Engine/Plugins/Media/PixelStreaming2/Source/PixelStreaming2RTC/Private/EpicRtcStreamer.h b/Engine/Plugins/Media/PixelStreaming2/Source/PixelStreaming2RTC/Private/EpicRtcStreamer.h
--- a/Engine/Plugins/Media/PixelStreaming2/Source/PixelStreaming2RTC/Private/EpicRtcStreamer.h
+++ b/Engine/Plugins/Media/PixelStreaming2/Source/PixelStreaming2RTC/Private/EpicRtcStreamer.h
@@ -2,15 +2,7 @@
 
 #pragma once
 
-#include "IPixelStreaming2Streamer.h"
-#include "VideoSourceGroup.h"
-#include "ThreadSafeMap.h"
 #include "Dom/JsonObject.h"
-#include "IPixelStreaming2InputHandler.h"
-#include "Templates/SharedPointer.h"
-#include "PlayerContext.h"
-#include "FreezeFrame.h"
-#include "StreamerReconnectTimer.h"
 #include "EpicRtcAudioTrackObserver.h"
 #include "EpicRtcAudioTrackObserverFactory.h"
 #include "EpicRtcDataTrackObserver.h"
@@ -19,6 +11,17 @@
 #include "EpicRtcVideoTrackObserverFactory.h"
 #include "EpicRtcRoomObserver.h"
 #include "EpicRtcSessionObserver.h"
+#include "FreezeFrame.h"
+#include "IPixelStreaming2InputHandler.h"
+#include "IPixelStreaming2Streamer.h"
+#include "Logging.h"
+#include "PlayerContext.h"
+#include "StreamerReconnectTimer.h"
+#include "Templates/SharedPointer.h"
+#include "ThreadSafeMap.h"
+#include "UtilsString.h"
+#include "VideoSourceGroup.h"
+
 
 #include "epic_rtc/core/conference.h"
 #include "epic_rtc/core/stats.h"
@@ -89,7 +92,7 @@ namespace UE::PixelStreaming2
 
 		virtual void RefreshStreamBitrate() override;
 
-		void ForEachPlayer(const TFunction<void(FString, FPlayerContext)>& Func);
+		void ForEachPlayer(const TFunction<void(FString, TSharedPtr<FPlayerContext>)>& Func);
 
 	private:
 		// own methods
@@ -97,19 +100,20 @@ namespace UE::PixelStreaming2
 		void ConsumeStats(FString PlayerId, FName StatName, float StatValue);
 		void DeletePlayerSession(FString PlayerId);
 		void DeleteAllPlayerSessions();
-		void OnDataChannelOpen(FString PlayerId);
-		void OnDataChannelClosed(FString PlayerId);
 		void SendInitialSettings(FString PlayerId) const;
 		void SendProtocol(FString PlayerId) const;
 		void SendPeerControllerMessages(FString PlayerId) const;
 		void SendLatencyReport(FString PlayerId) const;
 		void HandleRelayStatusMessage(const uint8_t* Data, uint32_t Size, EpicRtcDataTrackInterface* DataTrack);
-		bool FindPlayerByAudioTrack(EpicRtcAudioTrackInterface* AudioTrack, FString& OutPlayerId);
-		bool FindPlayerByVideoTrack(EpicRtcVideoTrackInterface* VideoTrack, FString& OutPlayerId);
-		bool FindPlayerByDataTrack(EpicRtcDataTrackInterface* DataTrack, FString& OutPlayerId);
 		void TriggerMouseLeave(FString InStreamerId);
-		void RemoveAudioTrack(EpicRtcAudioTrackInterface* AudioTrack);
-		void RemoveVideoTrack(EpicRtcVideoTrackInterface* VideoTrack);
+
+		// Function called on track events. Handles broadcasting of delegates and any other logic not related to the track itself (ie sending protocol)
+		void OnDataTrackOpen(FString PlayerId);
+		void OnDataTrackClosed(FString PlayerId);
+		void OnAudioTrackOpen(FString PlayerId, bool bIsRemote);
+		void OnAudioTrackClosed(FString PlayerId, bool bIsRemote);
+		void OnVideoTrackOpen(FString PlayerId, bool bIsRemote);
+		void OnVideoTrackClosed(FString PlayerId, bool bIsRemote);
 
 		void		   OnStatsReady(const FString& PlayerId, const EpicRtcConnectionStats& ConnectionStats);
 		void		   OnFrameCapturerCreated();
@@ -118,13 +122,57 @@ namespace UE::PixelStreaming2
 		void		   OnUseMediaCaptureChanged(IConsoleVariable* CVar);
 		EpicRtcBitrate GetBitrates() const;
 
+		template <typename T>
+		bool FindPlayerFromTrack(T Track, FString& OutPlayerId)
+		{
+			OutPlayerId = "";
+			if constexpr (std::is_same_v<T, EpicRtcVideoTrackInterface*>)
+			{
+				FString* FoundPlayerId = VideoTrackPlayerIdMap.Find(reinterpret_cast<uintptr_t>(Track));
+				if (FoundPlayerId)
+				{
+					OutPlayerId = *FoundPlayerId;
+				}
+			}
+			else if constexpr (std::is_same_v<T, EpicRtcAudioTrackInterface*>)
+			{
+				FString* FoundPlayerId = AudioTrackPlayerIdMap.Find(reinterpret_cast<uintptr_t>(Track));
+				if (FoundPlayerId)
+				{
+					OutPlayerId = *FoundPlayerId;
+				}
+			}
+			else if constexpr (std::is_same_v<T, EpicRtcDataTrackInterface*>)
+			{
+				FString DataTrackId = ToString(Track->GetId());
+				Participants->ApplyUntil([DataTrackId, &OutPlayerId](FString PlayerId, TSharedPtr<FPlayerContext> Participant) {
+					if (Participant->DataTrack)
+					{
+						FString TrackId = ToString(Participant->DataTrack->GetId());
+						if (TrackId == DataTrackId)
+						{
+							OutPlayerId = PlayerId;
+							return true;
+						}
+					}
+					return false;
+				});
+			}
+			else
+			{
+				UE_LOGFMT(LogPixelStreaming2RTC, Error, "FindPlayerFromTrack using unknown type. Ensure you're using a track interface pointer");
+			}
+
+			return !OutPlayerId.IsEmpty();
+		}
+
 	private:
 		FString StreamerId;
 		FString CurrentSignallingServerURL;
 
 		TSharedPtr<IPixelStreaming2InputHandler> InputHandler;
 
-		TSharedPtr<TThreadSafeMap<FString, FPlayerContext>> Players;
+		TSharedPtr<TThreadSafeMap<FString, TSharedPtr<FPlayerContext>>> Participants;
 
 		FString InputControllingId = INVALID_PLAYER_ID;
 
@@ -146,9 +194,9 @@ namespace UE::PixelStreaming2
 		TSharedPtr<FVideoSourceGroup>	  VideoSourceGroup;
 		TSharedPtr<FFreezeFrame>		  FreezeFrame;
 
-		TMap<FName, FString> ConfigOptions;
-		TMap<uintptr_t /* EpicRtcAudioTrackInterface* */,  FString /*participant id*/> AudioTrackPlayerIdMap;
-		TMap<uintptr_t /* EpicRtcVideoTrackInterface* */,  FString /*participant id*/> VideoTrackPlayerIdMap;
+		TMap<FName, FString>														  ConfigOptions;
+		TMap<uintptr_t /* EpicRtcAudioTrackInterface* */, FString /*participant id*/> AudioTrackPlayerIdMap;
+		TMap<uintptr_t /* EpicRtcVideoTrackInterface* */, FString /*participant id*/> VideoTrackPlayerIdMap;
 
 		TSharedPtr<FStreamerReconnectTimer> ReconnectTimer;
 
@@ -217,6 +265,7 @@ namespace UE::PixelStreaming2
 
 		virtual FString								 GetStreamType() override;
 		virtual TSharedPtr<IPixelStreaming2Streamer> CreateNewStreamer(const FString& StreamerId) override;
+
 	private:
 		TRefCountPtr<EpicRtcConferenceInterface> EpicRtcConference;
 	};
diff --git a/Engine/Plugins/Media/PixelStreaming2/Source/PixelStreaming2RTC/Private/EpicRtcVideoSink.cpp b/Engine/Plugins/Media/PixelStreaming2/Source/PixelStreaming2RTC/Private/EpicRtcVideoSink.cpp
--- a/Engine/Plugins/Media/PixelStreaming2/Source/PixelStreaming2RTC/Private/EpicRtcVideoSink.cpp
+++ b/Engine/Plugins/Media/PixelStreaming2/Source/PixelStreaming2RTC/Private/EpicRtcVideoSink.cpp
@@ -21,14 +21,14 @@ namespace UE::PixelStreaming2
 		Track = InTrack;
 	}
 
-	void FEpicRtcVideoSink::OnVideoData(const EpicRtcVideoFrame& Frame)
+	void FEpicRtcVideoSink::OnEpicRtcFrame(const EpicRtcVideoFrame& Frame)
 	{
-		if (!HasVideoConsumers() || bIsMuted)
+		if (!HasVideoConsumers() || bIsMuted || IsEngineExitRequested())
 		{
 			return;
 		}
 
-		TRACE_CPUPROFILER_EVENT_SCOPE_ON_CHANNEL_STR("FVideoSink::OnVideoData", PixelStreaming2Channel);
+		TRACE_CPUPROFILER_EVENT_SCOPE_ON_CHANNEL_STR("FEpicRtcVideoSink::OnEpicRtcFrame", PixelStreaming2Channel);
 
 		int32_t Width = Frame._buffer->GetWidth();
 		int32_t Height = Frame._buffer->GetHeight();
@@ -48,7 +48,7 @@ namespace UE::PixelStreaming2
 			}
 
 			auto& Raw = StaticCastSharedPtr<FVideoResourceRHI>(VideoResource)->GetRaw();
-			CallConsumeFrame(Raw.Texture);
+			OnVideoData(Raw.Texture);
 		}
 		else if (Frame._buffer->GetFormat() == EpicRtcPixelFormat::I420)
 		{
@@ -123,7 +123,7 @@ namespace UE::PixelStreaming2
 				// Set the Pixel data of the webrtc Frame to the SourceTexture
 				RHIUpdateTexture2D(SourceTexture, 0, Region, FrameSize.X * 4, Buffer.GetData());
 
-				CallConsumeFrame(SourceTexture);
+				OnVideoData(SourceTexture);
 			});
 		}
 	}
diff --git a/Engine/Plugins/Media/PixelStreaming2/Source/PixelStreaming2RTC/Private/FreezeFrame.cpp b/Engine/Plugins/Media/PixelStreaming2/Source/PixelStreaming2RTC/Private/FreezeFrame.cpp
--- a/Engine/Plugins/Media/PixelStreaming2/Source/PixelStreaming2RTC/Private/FreezeFrame.cpp
+++ b/Engine/Plugins/Media/PixelStreaming2/Source/PixelStreaming2RTC/Private/FreezeFrame.cpp
@@ -25,13 +25,13 @@
 
 namespace UE::PixelStreaming2
 {
-	TSharedPtr<FFreezeFrame> FFreezeFrame::Create(TWeakPtr<TThreadSafeMap<FString, FPlayerContext>> InPlayers, TWeakPtr<FVideoCapturer> InVideoCapturer, TWeakPtr<IPixelStreaming2InputHandler> InInputHandler)
+	TSharedPtr<FFreezeFrame> FFreezeFrame::Create(TWeakPtr<TThreadSafeMap<FString, TSharedPtr<FPlayerContext>>> InPlayers, TWeakPtr<FVideoCapturer> InVideoCapturer, TWeakPtr<IPixelStreaming2InputHandler> InInputHandler)
 	{
 		TSharedPtr<FFreezeFrame> FreezeFrame = TSharedPtr<FFreezeFrame>(new FFreezeFrame(InPlayers, InVideoCapturer, InInputHandler));
 		return FreezeFrame;
 	}
 
-	FFreezeFrame::FFreezeFrame(TWeakPtr<TThreadSafeMap<FString, FPlayerContext>> WeakPlayers, TWeakPtr<FVideoCapturer> WeakVideoCapturer, TWeakPtr<IPixelStreaming2InputHandler> WeakInputHandler)
+	FFreezeFrame::FFreezeFrame(TWeakPtr<TThreadSafeMap<FString, TSharedPtr<FPlayerContext>>> WeakPlayers, TWeakPtr<FVideoCapturer> WeakVideoCapturer, TWeakPtr<IPixelStreaming2InputHandler> WeakInputHandler)
 		: WeakPlayers(WeakPlayers)
 		, VideoCapturer(WeakVideoCapturer)
 		, InputHandler(WeakInputHandler)
@@ -154,7 +154,7 @@ namespace UE::PixelStreaming2
 
 	void FFreezeFrame::StopFreeze()
 	{
-		TSharedPtr<TThreadSafeMap<FString, FPlayerContext>> Players = WeakPlayers.Pin();
+		TSharedPtr<TThreadSafeMap<FString, TSharedPtr<FPlayerContext>>> Players = WeakPlayers.Pin();
 
 		if (!Players)
 		{
@@ -163,14 +163,14 @@ namespace UE::PixelStreaming2
 
 		TWeakPtr<IPixelStreaming2InputHandler> WeakHandler = InputHandler;
 
-		Players->Apply([WeakHandler](FString PlayerId, FPlayerContext& PlayerContext) {
-			if (!PlayerContext.DataTrack)
+		Players->Apply([WeakHandler](FString PlayerId, TSharedPtr<FPlayerContext> Participant) {
+			if (!Participant->DataTrack)
 			{
 				return;
 			}
 			if (TSharedPtr<IPixelStreaming2InputHandler> Handler = WeakHandler.Pin())
 			{
-				PlayerContext.DataTrack->SendMessage(EPixelStreaming2FromStreamerMessage::UnfreezeFrame);
+				Participant->DataTrack->SendMessage(EPixelStreaming2FromStreamerMessage::UnfreezeFrame);
 			}
 		});
 
@@ -179,7 +179,7 @@ namespace UE::PixelStreaming2
 
 	void FFreezeFrame::SendCachedFreezeFrameTo(FString PlayerId) const
 	{
-		TSharedPtr<TThreadSafeMap<FString, FPlayerContext>> Players = WeakPlayers.Pin();
+		TSharedPtr<TThreadSafeMap<FString, TSharedPtr<FPlayerContext>>> Players = WeakPlayers.Pin();
 
 		if (!Players)
 		{
@@ -188,15 +188,15 @@ namespace UE::PixelStreaming2
 
 		if (CachedJpegBytes.Num() > 0)
 		{
-			if (const FPlayerContext* PlayerContext = Players->Find(PlayerId))
+			if (TSharedPtr<FPlayerContext> Participant = Players->FindRef(PlayerId); Participant.IsValid())
 			{
-				if (!PlayerContext->DataTrack)
+				if (!Participant->DataTrack)
 				{
 					return;
 				}
 				if (TSharedPtr<IPixelStreaming2InputHandler> Handler = InputHandler.Pin())
 				{
-					PlayerContext->DataTrack->SendArbitraryData(EPixelStreaming2FromStreamerMessage::FreezeFrame, CachedJpegBytes);
+					Participant->DataTrack->SendArbitraryData(EPixelStreaming2FromStreamerMessage::FreezeFrame, CachedJpegBytes);
 				}
 			}
 		}
@@ -204,7 +204,7 @@ namespace UE::PixelStreaming2
 
 	void FFreezeFrame::SendFreezeFrame(TArray<FColor> RawData, const FIntRect& Rect)
 	{
-		TSharedPtr<TThreadSafeMap<FString, FPlayerContext>> Players = WeakPlayers.Pin();
+		TSharedPtr<TThreadSafeMap<FString, TSharedPtr<FPlayerContext>>> Players = WeakPlayers.Pin();
 
 		if (!Players)
 		{
@@ -219,14 +219,14 @@ namespace UE::PixelStreaming2
 			TWeakPtr<IPixelStreaming2InputHandler> WeakHandler = InputHandler;
 			// Compress to a JPEG of the maximum possible quality.
 			const TArray64<uint8>& JpegBytes = ImageWrapper->GetCompressed(100);
-			Players->Apply([&JpegBytes, WeakHandler](FString PlayerId, FPlayerContext& PlayerContext) {
-				if (!PlayerContext.DataTrack)
+			Players->Apply([&JpegBytes, WeakHandler](FString PlayerId, TSharedPtr<FPlayerContext> Participant) {
+				if (!Participant->DataTrack)
 				{
 					return;
 				}
 				if (TSharedPtr<IPixelStreaming2InputHandler> Handler = WeakHandler.Pin())
 				{
-					PlayerContext.DataTrack->SendArbitraryData(EPixelStreaming2FromStreamerMessage::FreezeFrame, JpegBytes);
+					Participant->DataTrack->SendArbitraryData(EPixelStreaming2FromStreamerMessage::FreezeFrame, JpegBytes);
 				}
 			});
 			CachedJpegBytes = JpegBytes;
diff --git a/Engine/Plugins/Media/PixelStreaming2/Source/PixelStreaming2RTC/Private/FreezeFrame.h b/Engine/Plugins/Media/PixelStreaming2/Source/PixelStreaming2RTC/Private/FreezeFrame.h
--- a/Engine/Plugins/Media/PixelStreaming2/Source/PixelStreaming2RTC/Private/FreezeFrame.h
+++ b/Engine/Plugins/Media/PixelStreaming2/Source/PixelStreaming2RTC/Private/FreezeFrame.h
@@ -27,7 +27,7 @@ namespace UE::PixelStreaming2
 	class FFreezeFrame : public TSharedFromThis<FFreezeFrame>
 	{
 	public:
-		static TSharedPtr<FFreezeFrame> Create(TWeakPtr<TThreadSafeMap<FString, FPlayerContext>> InPlayers, TWeakPtr<FVideoCapturer> InVideoCapturer, TWeakPtr<IPixelStreaming2InputHandler> InInputHandler);
+		static TSharedPtr<FFreezeFrame> Create(TWeakPtr<TThreadSafeMap<FString, TSharedPtr<FPlayerContext>>> InPlayers, TWeakPtr<FVideoCapturer> InVideoCapturer, TWeakPtr<IPixelStreaming2InputHandler> InInputHandler);
 		virtual ~FFreezeFrame();
 
 		/**
@@ -39,7 +39,7 @@ namespace UE::PixelStreaming2
 		void SendCachedFreezeFrameTo(FString PlayerId) const;
 
 	protected:
-		FFreezeFrame(TWeakPtr<TThreadSafeMap<FString, FPlayerContext>> WeakPlayers, TWeakPtr<FVideoCapturer> VideoCapturer, TWeakPtr<IPixelStreaming2InputHandler> InputHandler);
+		FFreezeFrame(TWeakPtr<TThreadSafeMap<FString, TSharedPtr<FPlayerContext>>> WeakPlayers, TWeakPtr<FVideoCapturer> VideoCapturer, TWeakPtr<IPixelStreaming2InputHandler> InputHandler);
 
 	private:
 		void SendFreezeFrame(TArray<FColor> RawData, const FIntRect& Rect);
@@ -48,7 +48,7 @@ namespace UE::PixelStreaming2
 		void FreezeFrameCapture();
 
 	private:
-		TWeakPtr<TThreadSafeMap<FString, FPlayerContext>> WeakPlayers;
+		TWeakPtr<TThreadSafeMap<FString, TSharedPtr<FPlayerContext>>> WeakPlayers;
 
 		// Video input used to capture the frame if mode 2 is used.
 		TWeakPtr<FVideoCapturer> VideoCapturer;
diff --git a/Engine/Plugins/Media/PixelStreaming2/Source/PixelStreaming2RTC/Private/PlayerContext.h b/Engine/Plugins/Media/PixelStreaming2/Source/PixelStreaming2RTC/Private/PlayerContext.h
--- a/Engine/Plugins/Media/PixelStreaming2/Source/PixelStreaming2RTC/Private/PlayerContext.h
+++ b/Engine/Plugins/Media/PixelStreaming2/Source/PixelStreaming2RTC/Private/PlayerContext.h
@@ -18,6 +18,22 @@
 
 namespace UE::PixelStreaming2
 {
+	/**
+	 * NOTE: As all of these members are either ref counted, or shared pointers, usage MUST look something like
+	 * 
+	 * TSharedPtr<TargetClass> Member; <-- Target member declared outside scope
+	 * if(TSharedPtr<FPlayerContext> Participant = Particpants->FindRef(ParticipantId); Participant.IsValid())
+	 * {
+	 * 		Member = Participant->Member;
+	 * }
+	 * 
+	 * if (Member)
+	 * {
+	 * 		Member->DoFunc();
+	 * }
+	 * 
+	 * Using this pattern means that we won't be keeping all members of the player context alive longer than required
+	 */
 	struct FPlayerContext
 	{
 		TRefCountPtr<EpicRtcParticipantInterface> ParticipantInterface;
