[PS2] Fix: Crashes related to PixelStreaming2 public API due to raw pointers being returned. All public APIs now return some form of smart pointer

diff --git a/Engine/Plugins/Media/PixelStreaming2/Source/PixelStreaming2/Internal/AudioSink.h b/Engine/Plugins/Media/PixelStreaming2/Source/PixelStreaming2/Internal/AudioSink.h
--- a/Engine/Plugins/Media/PixelStreaming2/Source/PixelStreaming2/Internal/AudioSink.h
+++ b/Engine/Plugins/Media/PixelStreaming2/Source/PixelStreaming2/Internal/AudioSink.h
@@ -12,8 +12,8 @@ namespace UE::PixelStreaming2
 		// Note: destructor will call destroy on any attached audio consumers
 		virtual ~FAudioSink();
 
-		virtual void AddAudioConsumer(IPixelStreaming2AudioConsumer* AudioConsumer) override;
-		virtual void RemoveAudioConsumer(IPixelStreaming2AudioConsumer* AudioConsumer) override;
+		virtual void AddAudioConsumer(const TWeakPtrVariant<IPixelStreaming2AudioConsumer>& AudioConsumer) override;
+		virtual void RemoveAudioConsumer(const TWeakPtrVariant<IPixelStreaming2AudioConsumer>& AudioConsumer) override;
 
 		bool HasAudioConsumers();
 
@@ -22,7 +22,7 @@ namespace UE::PixelStreaming2
 		void OnAudioData(int16_t* AudioData, uint32 NumFrames, uint32 NumChannels, uint32 SampleRate);
 
 	protected:
-		TSet<IPixelStreaming2AudioConsumer*> AudioConsumers;
+		TArray<TWeakPtrVariant<IPixelStreaming2AudioConsumer>> AudioConsumers;
 
 	private:
 		FCriticalSection AudioConsumersCS;
diff --git a/Engine/Plugins/Media/PixelStreaming2/Source/PixelStreaming2/Internal/VideoSink.h b/Engine/Plugins/Media/PixelStreaming2/Source/PixelStreaming2/Internal/VideoSink.h
--- a/Engine/Plugins/Media/PixelStreaming2/Source/PixelStreaming2/Internal/VideoSink.h
+++ b/Engine/Plugins/Media/PixelStreaming2/Source/PixelStreaming2/Internal/VideoSink.h
@@ -14,8 +14,8 @@ namespace UE::PixelStreaming2
 	public:
 		virtual ~FVideoSink();
 
-		virtual void AddVideoConsumer(IPixelStreaming2VideoConsumer* VideoConsumer) override;
-		virtual void RemoveVideoConsumer(IPixelStreaming2VideoConsumer* VideoConsumer) override;
+		virtual void AddVideoConsumer(const TWeakPtrVariant<IPixelStreaming2VideoConsumer>& VideoConsumer) override;
+		virtual void RemoveVideoConsumer(const TWeakPtrVariant<IPixelStreaming2VideoConsumer>& VideoConsumer) override;
 
 		bool HasVideoConsumers();
 
@@ -24,8 +24,8 @@ namespace UE::PixelStreaming2
 		void OnVideoData(FTextureRHIRef Frame);
 
 	protected:
-		FCriticalSection					 VideoConsumersCS;
-		TSet<IPixelStreaming2VideoConsumer*> VideoConsumers;
+		FCriticalSection									   VideoConsumersCS;
+		TArray<TWeakPtrVariant<IPixelStreaming2VideoConsumer>> VideoConsumers;
 
 		bool bIsMuted = false;
 	};
diff --git a/Engine/Plugins/Media/PixelStreaming2/Source/PixelStreaming2/Private/AudioSink.cpp b/Engine/Plugins/Media/PixelStreaming2/Source/PixelStreaming2/Private/AudioSink.cpp
--- a/Engine/Plugins/Media/PixelStreaming2/Source/PixelStreaming2/Private/AudioSink.cpp
+++ b/Engine/Plugins/Media/PixelStreaming2/Source/PixelStreaming2/Private/AudioSink.cpp
@@ -13,33 +13,40 @@ namespace UE::PixelStreaming2
 		FScopeLock Lock(&AudioConsumersCS);
 		for (auto Iter = AudioConsumers.CreateIterator(); Iter; ++Iter)
 		{
-			IPixelStreaming2AudioConsumer* AudioConsumer = Iter.ElementIt->Value;
+			TWeakPtrVariant<IPixelStreaming2AudioConsumer> AudioConsumer = *Iter;
 			Iter.RemoveCurrent();
-			if (AudioConsumer != nullptr)
+			if (TStrongPtrVariant<IPixelStreaming2AudioConsumer> PinnedAudioConsumer = AudioConsumer.Pin(); PinnedAudioConsumer.IsValid())
 			{
-				AudioConsumer->OnConsumerRemoved();
+				PinnedAudioConsumer.Get()->OnConsumerRemoved();
 			}
 		}
 	}
 
-	void FAudioSink::AddAudioConsumer(IPixelStreaming2AudioConsumer* AudioConsumer)
+	void FAudioSink::AddAudioConsumer(const TWeakPtrVariant<IPixelStreaming2AudioConsumer>& AudioConsumer)
 	{
 		FScopeLock Lock(&AudioConsumersCS);
-		bool	   bAlreadyInSet = false;
-		AudioConsumers.Add(AudioConsumer, &bAlreadyInSet);
-		if (!bAlreadyInSet)
+
+		if (!AudioConsumers.Contains(AudioConsumer))
 		{
-			AudioConsumer->OnConsumerAdded();
+			AudioConsumers.Add(AudioConsumer);
+			if (TStrongPtrVariant<IPixelStreaming2AudioConsumer> PinnedAudioConsumer = AudioConsumer.Pin(); PinnedAudioConsumer.IsValid())
+			{
+				PinnedAudioConsumer.Get()->OnConsumerAdded();
+			}
 		}
 	}
 
-	void FAudioSink::RemoveAudioConsumer(IPixelStreaming2AudioConsumer* AudioConsumer)
+	void FAudioSink::RemoveAudioConsumer(const TWeakPtrVariant<IPixelStreaming2AudioConsumer>& AudioConsumer)
 	{
 		FScopeLock Lock(&AudioConsumersCS);
 		if (AudioConsumers.Contains(AudioConsumer))
 		{
 			AudioConsumers.Remove(AudioConsumer);
-			AudioConsumer->OnConsumerRemoved();
+
+			if (TStrongPtrVariant<IPixelStreaming2AudioConsumer> PinnedAudioConsumer = AudioConsumer.Pin(); PinnedAudioConsumer.IsValid())
+			{
+				PinnedAudioConsumer.Get()->OnConsumerAdded();
+			}
 		}
 	}
 
@@ -59,12 +66,15 @@ namespace UE::PixelStreaming2
 		{
 			return;
 		}
-		
+
 		// Iterate audio consumers and pass this data to their buffers
 		FScopeLock Lock(&AudioConsumersCS);
-		for (IPixelStreaming2AudioConsumer* AudioConsumer : AudioConsumers)
+		for (TWeakPtrVariant<IPixelStreaming2AudioConsumer> AudioConsumer : AudioConsumers)
 		{
-			AudioConsumer->ConsumeRawPCM(AudioData, SampleRate, NumChannels, NumFrames);
+			if (TStrongPtrVariant<IPixelStreaming2AudioConsumer> PinnedAudioConsumer = AudioConsumer.Pin(); PinnedAudioConsumer.IsValid())
+			{
+				PinnedAudioConsumer.Get()->ConsumeRawPCM(AudioData, SampleRate, NumChannels, NumFrames);
+			}
 		}
 	}
 
diff --git a/Engine/Plugins/Media/PixelStreaming2/Source/PixelStreaming2/Private/Blueprints/PixelStreaming2AudioComponent.cpp b/Engine/Plugins/Media/PixelStreaming2/Source/PixelStreaming2/Private/Blueprints/PixelStreaming2AudioComponent.cpp
--- a/Engine/Plugins/Media/PixelStreaming2/Source/PixelStreaming2/Private/Blueprints/PixelStreaming2AudioComponent.cpp
+++ b/Engine/Plugins/Media/PixelStreaming2/Source/PixelStreaming2/Private/Blueprints/PixelStreaming2AudioComponent.cpp
@@ -98,15 +98,20 @@ bool UPixelStreaming2AudioComponent::StreamerListenTo(FString StreamerId, FStrin
 	{
 		return false;
 	}
-	IPixelStreaming2AudioSink* CandidateSink = WillListenToAnyPlayer() ? Streamer->GetUnlistenedAudioSink() : Streamer->GetPeerAudioSink(FString(PlayerToHear));
 
-	if (CandidateSink == nullptr)
+	TWeakPtr<IPixelStreaming2AudioSink> CandidateSink = WillListenToAnyPlayer() ? Streamer->GetUnlistenedAudioSink() : Streamer->GetPeerAudioSink(FString(PlayerToHear));
+
+	if (!CandidateSink.IsValid())
 	{
 		return false;
 	}
 
 	AudioSink = CandidateSink;
-	AudioSink->AddAudioConsumer(this);
+
+	if (TSharedPtr<IPixelStreaming2AudioSink> PinnedSink = AudioSink.Pin())
+	{
+		PinnedSink->AddAudioConsumer(TWeakPtrVariant<IPixelStreaming2AudioConsumer>(this));
+	}
 
 	return true;
 }
@@ -120,10 +125,12 @@ void UPixelStreaming2AudioComponent::Reset()
 		SoundGenerator->bShouldGenerateAudio = false;
 		SoundGenerator->EmptyBuffers();
 	}
-	if (AudioSink)
+
+	if (TSharedPtr<IPixelStreaming2AudioSink> PinnedSink = AudioSink.Pin())
 	{
-		AudioSink->RemoveAudioConsumer(this);
+		PinnedSink->RemoveAudioConsumer(TWeakPtrVariant<IPixelStreaming2AudioConsumer>(this));
 	}
+
 	AudioSink = nullptr;
 }
 
diff --git a/Engine/Plugins/Media/PixelStreaming2/Source/PixelStreaming2/Private/Blueprints/PixelStreaming2AudioComponent.h b/Engine/Plugins/Media/PixelStreaming2/Source/PixelStreaming2/Private/Blueprints/PixelStreaming2AudioComponent.h
--- a/Engine/Plugins/Media/PixelStreaming2/Source/PixelStreaming2/Private/Blueprints/PixelStreaming2AudioComponent.h
+++ b/Engine/Plugins/Media/PixelStreaming2/Source/PixelStreaming2/Private/Blueprints/PixelStreaming2AudioComponent.h
@@ -62,7 +62,7 @@ public:
 	bool bAutoFindPeer;
 
 private:
-	IPixelStreaming2AudioSink*											  AudioSink;
+	TWeakPtr<IPixelStreaming2AudioSink>									  AudioSink;
 	TSharedPtr<UE::PixelStreaming2::FSoundGenerator, ESPMode::ThreadSafe> SoundGenerator;
 
 public:
diff --git a/Engine/Plugins/Media/PixelStreaming2/Source/PixelStreaming2/Private/Blueprints/PixelStreaming2VideoComponent.cpp b/Engine/Plugins/Media/PixelStreaming2/Source/PixelStreaming2/Private/Blueprints/PixelStreaming2VideoComponent.cpp
--- a/Engine/Plugins/Media/PixelStreaming2/Source/PixelStreaming2/Private/Blueprints/PixelStreaming2VideoComponent.cpp
+++ b/Engine/Plugins/Media/PixelStreaming2/Source/PixelStreaming2/Private/Blueprints/PixelStreaming2VideoComponent.cpp
@@ -77,18 +77,18 @@ bool UPixelStreaming2VideoComponent::StreamerWatch(FString StreamerId, FString P
 		return false;
 	}
 
-	IPixelStreaming2VideoSink* CandidateSink = WillWatchAnyPlayer() ? Streamer->GetUnwatchedVideoSink() : Streamer->GetPeerVideoSink(FString(PlayerToWatch));
+	TWeakPtr<IPixelStreaming2VideoSink> CandidateSink = WillWatchAnyPlayer() ? Streamer->GetUnwatchedVideoSink() : Streamer->GetPeerVideoSink(FString(PlayerToWatch));
 
-	if (CandidateSink == nullptr)
+	if (!CandidateSink.IsValid())
 	{
 		return false;
 	}
 
 	VideoSink = CandidateSink;
 
-	if (VideoConsumer)
+	if (TSharedPtr<IPixelStreaming2VideoSink> PinnedSink = VideoSink.Pin(); PinnedSink && VideoConsumer)
 	{
-		VideoSink->AddVideoConsumer(VideoConsumer);
+		PinnedSink->AddVideoConsumer(TWeakPtrVariant<IPixelStreaming2VideoConsumer>(VideoConsumer.Get()));
 		bIsWatchingPlayer = true;
 	}
 
@@ -99,11 +99,13 @@ void UPixelStreaming2VideoComponent::Reset()
 {
 	PlayerToWatch = FString();
 	StreamerToWatch = FString();
-	if (VideoSink && VideoConsumer)
+
+	if (TSharedPtr<IPixelStreaming2VideoSink> PinnedSink = VideoSink.Pin(); PinnedSink && VideoConsumer)
 	{
-		VideoSink->RemoveVideoConsumer(VideoConsumer);
+		PinnedSink->RemoveVideoConsumer(TWeakPtrVariant<IPixelStreaming2VideoConsumer>(VideoConsumer.Get()));
 		bIsWatchingPlayer = false;
 	}
+
 	VideoSink = nullptr;
 }
 
diff --git a/Engine/Plugins/Media/PixelStreaming2/Source/PixelStreaming2/Private/Blueprints/PixelStreaming2VideoComponent.h b/Engine/Plugins/Media/PixelStreaming2/Source/PixelStreaming2/Private/Blueprints/PixelStreaming2VideoComponent.h
--- a/Engine/Plugins/Media/PixelStreaming2/Source/PixelStreaming2/Private/Blueprints/PixelStreaming2VideoComponent.h
+++ b/Engine/Plugins/Media/PixelStreaming2/Source/PixelStreaming2/Private/Blueprints/PixelStreaming2VideoComponent.h
@@ -59,7 +59,7 @@ public:
 	TObjectPtr<UPixelStreaming2MediaTexture> VideoConsumer = nullptr;
 
 private:
-	IPixelStreaming2VideoSink* VideoSink;
+	TWeakPtr<IPixelStreaming2VideoSink> VideoSink;
 
 	bool bIsWatchingPlayer;
 
diff --git a/Engine/Plugins/Media/PixelStreaming2/Source/PixelStreaming2/Private/VideoSink.cpp b/Engine/Plugins/Media/PixelStreaming2/Source/PixelStreaming2/Private/VideoSink.cpp
--- a/Engine/Plugins/Media/PixelStreaming2/Source/PixelStreaming2/Private/VideoSink.cpp
+++ b/Engine/Plugins/Media/PixelStreaming2/Source/PixelStreaming2/Private/VideoSink.cpp
@@ -15,33 +15,40 @@ namespace UE::PixelStreaming2
 		FScopeLock Lock(&VideoConsumersCS);
 		for (auto Iter = VideoConsumers.CreateIterator(); Iter; ++Iter)
 		{
-			IPixelStreaming2VideoConsumer* VideoConsumer = Iter.ElementIt->Value;
+			TWeakPtrVariant<IPixelStreaming2VideoConsumer> VideoConsumer = *Iter;
 			Iter.RemoveCurrent();
-			if (VideoConsumer != nullptr)
+			if (TStrongPtrVariant<IPixelStreaming2VideoConsumer> PinnedVideoConsumer = VideoConsumer.Pin(); PinnedVideoConsumer.IsValid())
 			{
-				VideoConsumer->OnConsumerRemoved();
+				PinnedVideoConsumer.Get()->OnConsumerRemoved();
 			}
 		}
 	}
 
-	void FVideoSink::AddVideoConsumer(IPixelStreaming2VideoConsumer* VideoConsumer)
+	void FVideoSink::AddVideoConsumer(const TWeakPtrVariant<IPixelStreaming2VideoConsumer>& VideoConsumer)
 	{
 		FScopeLock Lock(&VideoConsumersCS);
-		bool	   bAlreadyInSet = false;
-		VideoConsumers.Add(VideoConsumer, &bAlreadyInSet);
-		if (!bAlreadyInSet)
+
+		if (!VideoConsumers.Contains(VideoConsumer))
 		{
-			VideoConsumer->OnConsumerAdded();
+			VideoConsumers.Add(VideoConsumer);
+			if (TStrongPtrVariant<IPixelStreaming2VideoConsumer> PinnedVideoConsumer = VideoConsumer.Pin(); PinnedVideoConsumer.IsValid())
+			{	
+				PinnedVideoConsumer.Get()->OnConsumerAdded();
+			}
 		}
 	}
 
-	void FVideoSink::RemoveVideoConsumer(IPixelStreaming2VideoConsumer* VideoConsumer)
+	void FVideoSink::RemoveVideoConsumer(const TWeakPtrVariant<IPixelStreaming2VideoConsumer>& VideoConsumer)
 	{
 		FScopeLock Lock(&VideoConsumersCS);
 		if (VideoConsumers.Contains(VideoConsumer))
 		{
 			VideoConsumers.Remove(VideoConsumer);
-			VideoConsumer->OnConsumerRemoved();
+
+			if (TStrongPtrVariant<IPixelStreaming2VideoConsumer> PinnedVideoConsumer = VideoConsumer.Pin(); PinnedVideoConsumer.IsValid())
+			{	
+				PinnedVideoConsumer.Get()->OnConsumerRemoved();
+			}
 		}
 	}
 
@@ -61,9 +68,12 @@ namespace UE::PixelStreaming2
 
 		// Iterate video consumers and pass this data to their buffers
 		FScopeLock Lock(&VideoConsumersCS);
-		for (IPixelStreaming2VideoConsumer* VideoConsumer : VideoConsumers)
+		for (TWeakPtrVariant<IPixelStreaming2VideoConsumer> VideoConsumer : VideoConsumers)
 		{
-			VideoConsumer->ConsumeFrame(Frame);
+			if (TStrongPtrVariant<IPixelStreaming2VideoConsumer> PinnedVideoConsumer = VideoConsumer.Pin(); PinnedVideoConsumer.IsValid())
+			{
+				PinnedVideoConsumer.Get()->ConsumeFrame(Frame);
+			}
 		}
 	}
 
diff --git a/Engine/Plugins/Media/PixelStreaming2/Source/PixelStreaming2Core/Private/Streamer.cpp b/Engine/Plugins/Media/PixelStreaming2/Source/PixelStreaming2Core/Private/Streamer.cpp
--- a/Engine/Plugins/Media/PixelStreaming2/Source/PixelStreaming2Core/Private/Streamer.cpp
+++ b/Engine/Plugins/Media/PixelStreaming2/Source/PixelStreaming2Core/Private/Streamer.cpp
@@ -40,10 +40,10 @@ namespace UE::PixelStreaming2
 		virtual void										 KickPlayer(FString PlayerId) override {}
 		virtual TArray<FString>								 GetConnectedPlayers() override { return {}; }
 		virtual TWeakPtr<class IPixelStreaming2InputHandler> GetInputHandler() override { return nullptr; }
-		virtual IPixelStreaming2AudioSink*					 GetPeerAudioSink(FString PlayerId) override { return nullptr; }
-		virtual IPixelStreaming2AudioSink*					 GetUnlistenedAudioSink() override { return nullptr; }
-		virtual IPixelStreaming2VideoSink*					 GetPeerVideoSink(FString PlayerId) override { return nullptr; }
-		virtual IPixelStreaming2VideoSink*					 GetUnwatchedVideoSink() override { return nullptr; }
+		virtual TWeakPtr<IPixelStreaming2AudioSink>			 GetPeerAudioSink(FString PlayerId) override { return nullptr; }
+		virtual TWeakPtr<IPixelStreaming2AudioSink>			 GetUnlistenedAudioSink() override { return nullptr; }
+		virtual TWeakPtr<IPixelStreaming2VideoSink>			 GetPeerVideoSink(FString PlayerId) override { return nullptr; }
+		virtual TWeakPtr<IPixelStreaming2VideoSink>			 GetUnwatchedVideoSink() override { return nullptr; }
 		virtual void										 SetConfigOption(const FName& OptionName, const FString& Value) override {}
 		virtual bool										 GetConfigOption(const FName& OptionName, FString& OutValue) override { return false; }
 		virtual void										 PlayerRequestsBitrate(FString PlayerId, int MinBitrate, int MaxBitrate) override {}
diff --git a/Engine/Plugins/Media/PixelStreaming2/Source/PixelStreaming2Core/Public/IPixelStreaming2AudioSink.h b/Engine/Plugins/Media/PixelStreaming2/Source/PixelStreaming2Core/Public/IPixelStreaming2AudioSink.h
--- a/Engine/Plugins/Media/PixelStreaming2/Source/PixelStreaming2Core/Public/IPixelStreaming2AudioSink.h
+++ b/Engine/Plugins/Media/PixelStreaming2/Source/PixelStreaming2Core/Public/IPixelStreaming2AudioSink.h
@@ -4,6 +4,7 @@
 
 #include "CoreTypes.h"
 #include "IPixelStreaming2AudioConsumer.h"
+#include "Templates/PointerVariants.h"
 #include "UObject/Interface.h"
 
 #include "IPixelStreaming2AudioSink.generated.h"
@@ -27,11 +28,17 @@ public:
 	 * @brief Add an audio consumer to the sink.
 	 * @param AudioConsumer The Audio consumer to add to the sink.
 	 */
-	virtual void AddAudioConsumer(IPixelStreaming2AudioConsumer* AudioConsumer) = 0;
+	UE_DEPRECATED(5.6, "AddAudioConsumer(IPixelStreaming2AudioConsumer*) has been deprecated. Please use AddAudioConsumer(const TWeakPtrVariant<IPixelStreaming2AudioConsumer>&) instead.")
+	virtual void AddAudioConsumer(IPixelStreaming2AudioConsumer* AudioConsumer) {}
+
+	virtual void AddAudioConsumer(const TWeakPtrVariant<IPixelStreaming2AudioConsumer>& AudioConsumer) = 0;
 
 	/**
 	 * @brief Remove an audio consumer to remove from the sink.
 	 * @param AudioConsumer The Audio consumer to remove from the sink.
 	 */
-	virtual void RemoveAudioConsumer(IPixelStreaming2AudioConsumer* AudioConsumer) = 0;		
+	UE_DEPRECATED(5.6, "RemoveAudioConsumer(IPixelStreaming2AudioConsumer*) has been deprecated. Please use RemoveAudioConsumer(const TWeakPtrVariant<IPixelStreaming2AudioConsumer>&) instead.")
+	virtual void RemoveAudioConsumer(IPixelStreaming2AudioConsumer* AudioConsumer) {}
+
+	virtual void RemoveAudioConsumer(const TWeakPtrVariant<IPixelStreaming2AudioConsumer>& AudioConsumer) = 0;
 };
\ No newline at end of file
diff --git a/Engine/Plugins/Media/PixelStreaming2/Source/PixelStreaming2Core/Public/IPixelStreaming2Streamer.h b/Engine/Plugins/Media/PixelStreaming2/Source/PixelStreaming2Core/Public/IPixelStreaming2Streamer.h
--- a/Engine/Plugins/Media/PixelStreaming2/Source/PixelStreaming2Core/Public/IPixelStreaming2Streamer.h
+++ b/Engine/Plugins/Media/PixelStreaming2/Source/PixelStreaming2Core/Public/IPixelStreaming2Streamer.h
@@ -219,25 +219,25 @@ public:
 	 * @param PlayerId The player id of the audio sink to retrieve.
 	 * @return The PeerAudioSink.
 	 */
-	virtual IPixelStreaming2AudioSink* GetPeerAudioSink(FString PlayerId) = 0;
+	virtual TWeakPtr<IPixelStreaming2AudioSink> GetPeerAudioSink(FString PlayerId) = 0;
 
 	/**
 	 * @brief Get an audio sink that has no peers/players listening to it.
 	 * @return The unlistened PeerAudioSink.
 	 */
-	virtual IPixelStreaming2AudioSink* GetUnlistenedAudioSink() = 0;
+	virtual TWeakPtr<IPixelStreaming2AudioSink> GetUnlistenedAudioSink() = 0;
 
 	/**
 	 * @brief Get the video sink associated with a specific peer/player.
 	 * @return The PeerVideoSink.
 	 */
-	virtual IPixelStreaming2VideoSink* GetPeerVideoSink(FString PlayerId) = 0;
+	virtual TWeakPtr<IPixelStreaming2VideoSink> GetPeerVideoSink(FString PlayerId) = 0;
 
 	/**
 	 * @brief Get a video sink that has no peers/players watching it.
 	 * @return The unwatched PeerVideoSink.
 	 */
-	virtual IPixelStreaming2VideoSink* GetUnwatchedVideoSink() = 0;
+	virtual TWeakPtr<IPixelStreaming2VideoSink> GetUnwatchedVideoSink() = 0;
 
 	/**
 	 * @brief Allows sending arbitrary configuration options during initial connection.
diff --git a/Engine/Plugins/Media/PixelStreaming2/Source/PixelStreaming2Core/Public/IPixelStreaming2VideoSink.h b/Engine/Plugins/Media/PixelStreaming2/Source/PixelStreaming2Core/Public/IPixelStreaming2VideoSink.h
--- a/Engine/Plugins/Media/PixelStreaming2/Source/PixelStreaming2Core/Public/IPixelStreaming2VideoSink.h
+++ b/Engine/Plugins/Media/PixelStreaming2/Source/PixelStreaming2Core/Public/IPixelStreaming2VideoSink.h
@@ -4,6 +4,7 @@
 
 #include "CoreTypes.h"
 #include "IPixelStreaming2VideoConsumer.h"
+#include "Templates/PointerVariants.h"
 #include "UObject/Interface.h"
 
 #include "IPixelStreaming2VideoSink.generated.h"
@@ -28,11 +29,25 @@ public:
 	 * @brief Add a video consumer to the sink.
 	 * @param VideoConsumer The video consumer to add to the sink.
 	 */
-	virtual void AddVideoConsumer(IPixelStreaming2VideoConsumer* VideoConsumer) = 0;
+	UE_DEPRECATED(5.6, "AddVideoConsumer(IPixelStreaming2VideoConsumer*) has been deprecated. Please use AddVideoConsumer(const TWeakPtrVariant<IPixelStreaming2VideoConsumer>&) instead.")
+	virtual void AddVideoConsumer(IPixelStreaming2VideoConsumer* VideoConsumer) {}
+
+	/**
+	 * @brief Add a video consumer to the sink.
+	 * @param VideoConsumer The video consumer to add to the sink.
+	 */
+	virtual void AddVideoConsumer(const TWeakPtrVariant<IPixelStreaming2VideoConsumer>& VideoConsumer) = 0;
+
+	/**
+	 * @brief Remove a video consumer to remove from the sink.
+	 * @param VideoConsumer The video consumer to remove from the sink.
+	 */
+	UE_DEPRECATED(5.6, "RemoveVideoConsumer(IPixelStreaming2VideoConsumer*) has been deprecated. Please use RemoveVideoConsumer(const TWeakPtrVariant<IPixelStreaming2VideoConsumer>&) instead.")
+	virtual void RemoveVideoConsumer(IPixelStreaming2VideoConsumer* VideoConsumer) {}
 
 	/**
 	 * @brief Remove a video consumer to remove from the sink.
 	 * @param VideoConsumer The video consumer to remove from the sink.
 	 */
-	virtual void RemoveVideoConsumer(IPixelStreaming2VideoConsumer* VideoConsumer) = 0;
+	virtual void RemoveVideoConsumer(const TWeakPtrVariant<IPixelStreaming2VideoConsumer>& VideoConsumer) = 0;
 };
diff --git a/Engine/Plugins/Media/PixelStreaming2/Source/PixelStreaming2RTC/Private/Blueprints/PixelStreaming2Peer.cpp b/Engine/Plugins/Media/PixelStreaming2/Source/PixelStreaming2RTC/Private/Blueprints/PixelStreaming2Peer.cpp
--- a/Engine/Plugins/Media/PixelStreaming2/Source/PixelStreaming2RTC/Private/Blueprints/PixelStreaming2Peer.cpp
+++ b/Engine/Plugins/Media/PixelStreaming2/Source/PixelStreaming2RTC/Private/Blueprints/PixelStreaming2Peer.cpp
@@ -28,7 +28,7 @@ UPixelStreaming2Peer::UPixelStreaming2Peer(const FObjectInitializer& ObjectIniti
 void UPixelStreaming2Peer::BeginPlay()
 {
 	Super::BeginPlay();
-	
+
 	SessionObserver = MakeRefCount<UE::PixelStreaming2::FEpicRtcSessionObserver>(UE::PixelStreaming2::TObserver<IPixelStreaming2SessionObserver>(this));
 	RoomObserver = MakeRefCount<UE::PixelStreaming2::FEpicRtcRoomObserver>(UE::PixelStreaming2::TObserver<IPixelStreaming2RoomObserver>(this));
 
@@ -36,6 +36,8 @@ void UPixelStreaming2Peer::BeginPlay()
 	VideoTrackObserverFactory = MakeRefCount<UE::PixelStreaming2::FEpicRtcVideoTrackObserverFactory>(UE::PixelStreaming2::TObserver<IPixelStreaming2VideoTrackObserver>(this));
 	DataTrackObserverFactory = MakeRefCount<UE::PixelStreaming2::FEpicRtcDataTrackObserverFactory>(UE::PixelStreaming2::TObserver<IPixelStreaming2DataTrackObserver>(this));
 
+	EpicRtcConference = UE::PixelStreaming2::FPixelStreaming2RTCModule::GetModule()->GetEpicRtcConference();
+
 	UE::PixelStreaming2::FPixelStreaming2RTCModule::GetModule()->GetStatsCollector()->OnStatsReady.AddUObject(this, &UPixelStreaming2Peer::OnStatsReady);
 }
 
@@ -64,15 +66,13 @@ bool UPixelStreaming2Peer::Connect(const FString& Url)
 		._observer = SessionObserver.GetReference()
 	};
 
-	EpicRtcErrorCode Result = UE::PixelStreaming2::FPixelStreaming2RTCModule::GetModule()->GetEpicRtcConference()->CreateSession(SessionConfig, EpicRtcSession.GetInitReference());
-	if (Result != EpicRtcErrorCode::Ok)
+	if (const EpicRtcErrorCode Result = EpicRtcConference->CreateSession(SessionConfig, EpicRtcSession.GetInitReference()); Result != EpicRtcErrorCode::Ok)
 	{
 		UE_LOGFMT(LogPixelStreaming2RTC, Error, "Failed to create EpicRtc session, CreateSession returned: {0}", UE::PixelStreaming2::ToString(Result));
 		return false;
 	}
 
-	Result = EpicRtcSession->Connect();
-	if (Result != EpicRtcErrorCode::Ok)
+	if (const EpicRtcErrorCode Result = EpicRtcSession->Connect(); Result != EpicRtcErrorCode::Ok)
 	{
 		UE_LOGFMT(LogPixelStreaming2RTC, Error, "Failed to connect EpicRtcSession. Connect returned: {0}", UE::PixelStreaming2::ToString(Result));
 		return false;
@@ -93,6 +93,26 @@ bool UPixelStreaming2Peer::Disconnect(const FString& OptionalReason)
 		return false;
 	}
 
+	if (AudioSink)
+	{
+		AudioSink->RemoveAudioConsumer(this);
+		AudioSink = nullptr;
+		RemoteAudioTrack = 0;
+	}
+
+	// NOTE: It is imperative we null out the video sink before we remove room and session. If the sink is still alive
+	// during session destruction, webrtc will flush incoming frames (leading to OnVideoTrackFrame) and the engine will
+	// lock inside AVCodecs on a RHI fence. With the sink nulled, OnVideoTrackFrame will early exit
+	if (VideoSink)
+	{
+		if (VideoConsumer)
+		{
+			VideoSink->RemoveVideoConsumer(TWeakPtrVariant<IPixelStreaming2VideoConsumer>(VideoConsumer.Get()));
+		}
+		VideoSink = nullptr;
+		RemoteVideoTrack = 0;
+	}
+
 	if (EpicRtcRoom)
 	{
 		EpicRtcRoom->Leave();
@@ -116,17 +136,7 @@ bool UPixelStreaming2Peer::Disconnect(const FString& OptionalReason)
 		return false;
 	}
 
-	if (AudioSink)
-	{
-		AudioSink->RemoveAudioConsumer(this);
-	}
-
-	if (VideoSink && VideoConsumer)
-	{
-		VideoSink->RemoveVideoConsumer(VideoConsumer);
-	}
-
-	UE::PixelStreaming2::FPixelStreaming2RTCModule::GetModule()->GetEpicRtcConference()->RemoveSession({ ._ptr = reinterpret_cast<const char*>(*PlayerName), ._length = static_cast<uint64_t>(PlayerName.Len()) });
+	EpicRtcConference->RemoveSession({ ._ptr = reinterpret_cast<const char*>(*PlayerName), ._length = static_cast<uint64_t>(PlayerName.Len()) });
 
 	return true;
 }
@@ -291,6 +301,7 @@ void UPixelStreaming2Peer::OnAudioTrackUpdate(EpicRtcParticipantInterface* Parti
 	{
 		// We received a remote track. We should now generate audio from it
 		AudioSink = UE::PixelStreaming2::FEpicRtcAudioSink::Create(AudioTrack);
+		RemoteAudioTrack = reinterpret_cast<uintptr_t>(AudioTrack); // Keep track of which remote track we're receiving audio from as we only support one track
 		AudioSink->AddAudioConsumer(this);
 	}
 }
@@ -305,9 +316,10 @@ void UPixelStreaming2Peer::OnVideoTrackUpdate(EpicRtcParticipantInterface* Parti
 	{
 		// We received a remote track. We should now generate video from it
 		VideoSink = UE::PixelStreaming2::FEpicRtcVideoSink::Create(VideoTrack);
+		RemoteVideoTrack = reinterpret_cast<uintptr_t>(VideoTrack); // Keep track of which remote track we're receiving video from as we only support display one track
 		if (VideoConsumer)
 		{
-			VideoSink->AddVideoConsumer(VideoConsumer);
+			VideoSink->AddVideoConsumer(TWeakPtrVariant<IPixelStreaming2VideoConsumer>(VideoConsumer.Get()));
 		}
 	}
 }
@@ -332,7 +344,7 @@ void UPixelStreaming2Peer::OnRoomErrorUpdate(const EpicRtcErrorCode Error)
 
 void UPixelStreaming2Peer::OnAudioTrackMuted(EpicRtcAudioTrackInterface* AudioTrack, EpicRtcBool bIsMuted)
 {
-	if (!AudioSink)
+	if (!AudioSink || RemoteAudioTrack != reinterpret_cast<uintptr_t>(AudioTrack))
 	{
 		return;
 	}
@@ -342,7 +354,7 @@ void UPixelStreaming2Peer::OnAudioTrackMuted(EpicRtcAudioTrackInterface* AudioTr
 
 void UPixelStreaming2Peer::OnAudioTrackFrame(EpicRtcAudioTrackInterface* AudioTrack, const EpicRtcAudioFrame& Frame)
 {
-	if (!AudioSink)
+	if (!AudioSink || RemoteAudioTrack != reinterpret_cast<uintptr_t>(AudioTrack))
 	{
 		return;
 	}
@@ -360,7 +372,7 @@ void UPixelStreaming2Peer::OnAudioTrackState(EpicRtcAudioTrackInterface* AudioTr
 
 void UPixelStreaming2Peer::OnVideoTrackMuted(EpicRtcVideoTrackInterface* VideoTrack, EpicRtcBool bIsMuted)
 {
-	if (!VideoSink)
+	if (!VideoSink || RemoteVideoTrack != reinterpret_cast<uintptr_t>(VideoTrack))
 	{
 		return;
 	}
@@ -370,7 +382,7 @@ void UPixelStreaming2Peer::OnVideoTrackMuted(EpicRtcVideoTrackInterface* VideoTr
 
 void UPixelStreaming2Peer::OnVideoTrackFrame(EpicRtcVideoTrackInterface* VideoTrack, const EpicRtcVideoFrame& Frame)
 {
-	if (!VideoSink)
+	if (!VideoSink || RemoteVideoTrack != reinterpret_cast<uintptr_t>(VideoTrack))
 	{
 		return;
 	}
diff --git a/Engine/Plugins/Media/PixelStreaming2/Source/PixelStreaming2RTC/Private/Blueprints/PixelStreaming2Peer.h b/Engine/Plugins/Media/PixelStreaming2/Source/PixelStreaming2RTC/Private/Blueprints/PixelStreaming2Peer.h
--- a/Engine/Plugins/Media/PixelStreaming2/Source/PixelStreaming2RTC/Private/Blueprints/PixelStreaming2Peer.h
+++ b/Engine/Plugins/Media/PixelStreaming2/Source/PixelStreaming2RTC/Private/Blueprints/PixelStreaming2Peer.h
@@ -99,10 +99,13 @@ private:
 	FUtf8String		PlayerName;
 	static uint32_t PlayerId;
 
-	TSharedPtr<UE::PixelStreaming2::FEpicRtcAudioSink>					  AudioSink;
+	TSharedPtr<UE::PixelStreaming2::FEpicRtcAudioSink> AudioSink;
+	uintptr_t										   RemoteAudioTrack;
+
 	TSharedPtr<UE::PixelStreaming2::FSoundGenerator, ESPMode::ThreadSafe> SoundGenerator;
 
 	TSharedPtr<UE::PixelStreaming2::FEpicRtcVideoSink> VideoSink;
+	uintptr_t										   RemoteVideoTrack;
 
 	TSharedPtr<UE::PixelStreaming2::FRTCStatsCollector> StatsCollector;
 
@@ -161,8 +164,9 @@ private:
 	EpicRtcSessionState SessionState = EpicRtcSessionState::Disconnected;
 
 	// Begin EpicRtc Classes
-	TRefCountPtr<EpicRtcSessionInterface> EpicRtcSession;
-	TRefCountPtr<EpicRtcRoomInterface>	  EpicRtcRoom;
+	TRefCountPtr<EpicRtcConferenceInterface> EpicRtcConference;
+	TRefCountPtr<EpicRtcSessionInterface>	 EpicRtcSession;
+	TRefCountPtr<EpicRtcRoomInterface>		 EpicRtcRoom;
 	// End EpicRtc Classes
 
 	// Begin EpicRtc Observers
diff --git a/Engine/Plugins/Media/PixelStreaming2/Source/PixelStreaming2RTC/Private/EpicRtcStreamer.cpp b/Engine/Plugins/Media/PixelStreaming2/Source/PixelStreaming2RTC/Private/EpicRtcStreamer.cpp
--- a/Engine/Plugins/Media/PixelStreaming2/Source/PixelStreaming2RTC/Private/EpicRtcStreamer.cpp
+++ b/Engine/Plugins/Media/PixelStreaming2/Source/PixelStreaming2RTC/Private/EpicRtcStreamer.cpp
@@ -499,26 +499,26 @@ namespace UE::PixelStreaming2
 		return ConnectedParticipantIds.Array();
 	}
 
-	IPixelStreaming2AudioSink* FEpicRtcStreamer::GetPeerAudioSink(FString PlayerId)
+	TWeakPtr<IPixelStreaming2AudioSink> FEpicRtcStreamer::GetPeerAudioSink(FString PlayerId)
 	{
-		IPixelStreaming2AudioSink* Result = nullptr;
+		TWeakPtr<IPixelStreaming2AudioSink> Result = nullptr;
 		if (TSharedPtr<FPlayerContext> Participant = Participants->FindRef(PlayerId); Participant.IsValid())
 		{
-			Result = Participant->AudioSink ? Participant->AudioSink.Get() : nullptr;
+			Result = Participant->AudioSink ? Participant->AudioSink : nullptr;
 		}
 
 		return Result;
 	}
 
-	IPixelStreaming2AudioSink* FEpicRtcStreamer::GetUnlistenedAudioSink()
+	TWeakPtr<IPixelStreaming2AudioSink> FEpicRtcStreamer::GetUnlistenedAudioSink()
 	{
-		IPixelStreaming2AudioSink* Result = nullptr;
+		TWeakPtr<IPixelStreaming2AudioSink> Result = nullptr;
 		Participants->ApplyUntil([&Result](FString PlayerId, TSharedPtr<FPlayerContext> Participant) {
 			if (Participant->AudioSink)
 			{
 				if (!Participant->AudioSink->HasAudioConsumers())
 				{
-					Result = Participant->AudioSink.Get();
+					Result = Participant->AudioSink;
 					return true;
 				}
 			}
@@ -527,26 +527,26 @@ namespace UE::PixelStreaming2
 		return Result;
 	}
 
-	IPixelStreaming2VideoSink* FEpicRtcStreamer::GetPeerVideoSink(FString PlayerId)
+	TWeakPtr<IPixelStreaming2VideoSink> FEpicRtcStreamer::GetPeerVideoSink(FString PlayerId)
 	{
-		IPixelStreaming2VideoSink* Result = nullptr;
+		TWeakPtr<IPixelStreaming2VideoSink> Result = nullptr;
 		if (TSharedPtr<FPlayerContext> Participant = Participants->FindRef(PlayerId); Participant.IsValid())
 		{
-			Result = Participant->VideoSink ? Participant->VideoSink.Get() : nullptr;
+			Result = Participant->VideoSink ? Participant->VideoSink : nullptr;
 		}
 
 		return Result;
 	}
 
-	IPixelStreaming2VideoSink* FEpicRtcStreamer::GetUnwatchedVideoSink()
+	TWeakPtr<IPixelStreaming2VideoSink> FEpicRtcStreamer::GetUnwatchedVideoSink()
 	{
-		IPixelStreaming2VideoSink* Result = nullptr;
+		TWeakPtr<IPixelStreaming2VideoSink> Result = nullptr;
 		Participants->ApplyUntil([&Result](FString PlayerId, TSharedPtr<FPlayerContext> Participant) {
 			if (Participant->VideoSink)
 			{
 				if (!Participant->VideoSink->HasVideoConsumers())
 				{
-					Result = Participant->VideoSink.Get();
+					Result = Participant->VideoSink;
 					return true;
 				}
 			}
@@ -640,10 +640,10 @@ namespace UE::PixelStreaming2
 	{
 		// We need to keep a copy of the participant to delete as when destroying the ParticipantInterface,
 		// webrtc will finish tasks in the other threads which can execute FindRef and end up deadlocking
-		// the map. 
+		// the map.
 		// TL:DR Participant deletion needs to happen outside the TThreadSafeMap scope lock
 		TSharedPtr<FPlayerContext> ParticipantToDelete;
-		if(!Participants->RemoveAndCopyValue(PlayerId, ParticipantToDelete) || !ParticipantToDelete.IsValid())
+		if (!Participants->RemoveAndCopyValue(PlayerId, ParticipantToDelete) || !ParticipantToDelete.IsValid())
 		{
 			return;
 		}
@@ -1641,7 +1641,7 @@ namespace UE::PixelStreaming2
 
 		if (!bFoundPlayer)
 		{
-			UE_LOG(LogPixelStreaming2RTC, Warning, TEXT("FEpicRtcStreamer::OnAudioTrackFrame(Failed to find a player for audio track [%s])"), *AudioTrackId);
+			UE_LOG(LogPixelStreaming2RTC, Log, TEXT("FEpicRtcStreamer::OnAudioTrackFrame(Failed to find a player for audio track [%s])"), *AudioTrackId);
 			return;
 		}
 
@@ -1734,7 +1734,7 @@ namespace UE::PixelStreaming2
 
 		if (!bFoundPlayer)
 		{
-			UE_LOG(LogPixelStreaming2RTC, Warning, TEXT("FEpicRtcStreamer::OnVideoTrackFrame(Failed to find a player for video track [%s])"), *VideoTrackId);
+			UE_LOG(LogPixelStreaming2RTC, Log, TEXT("FEpicRtcStreamer::OnVideoTrackFrame(Failed to find a player for video track [%s])"), *VideoTrackId);
 			return;
 		}
 
diff --git a/Engine/Plugins/Media/PixelStreaming2/Source/PixelStreaming2RTC/Private/EpicRtcStreamer.h b/Engine/Plugins/Media/PixelStreaming2/Source/PixelStreaming2RTC/Private/EpicRtcStreamer.h
--- a/Engine/Plugins/Media/PixelStreaming2/Source/PixelStreaming2RTC/Private/EpicRtcStreamer.h
+++ b/Engine/Plugins/Media/PixelStreaming2/Source/PixelStreaming2RTC/Private/EpicRtcStreamer.h
@@ -80,10 +80,10 @@ namespace UE::PixelStreaming2
 
 		virtual TWeakPtr<IPixelStreaming2InputHandler> GetInputHandler() override { return InputHandler; }
 
-		virtual IPixelStreaming2AudioSink* GetPeerAudioSink(FString PlayerId) override;
-		virtual IPixelStreaming2AudioSink* GetUnlistenedAudioSink() override;
-		virtual IPixelStreaming2VideoSink* GetPeerVideoSink(FString PlayerId) override;
-		virtual IPixelStreaming2VideoSink* GetUnwatchedVideoSink() override;
+		virtual TWeakPtr<IPixelStreaming2AudioSink> GetPeerAudioSink(FString PlayerId) override;
+		virtual TWeakPtr<IPixelStreaming2AudioSink> GetUnlistenedAudioSink() override;
+		virtual TWeakPtr<IPixelStreaming2VideoSink> GetPeerVideoSink(FString PlayerId) override;
+		virtual TWeakPtr<IPixelStreaming2VideoSink> GetUnwatchedVideoSink() override;
 
 		virtual void SetConfigOption(const FName& OptionName, const FString& Value) override;
 		virtual bool GetConfigOption(const FName& OptionName, FString& OutValue) override;
diff --git a/Engine/Plugins/Media/PixelStreaming2/Source/PixelStreaming2RTC/Private/Tests/DelegateTest.cpp b/Engine/Plugins/Media/PixelStreaming2/Source/PixelStreaming2RTC/Private/Tests/DelegateTest.cpp
--- a/Engine/Plugins/Media/PixelStreaming2/Source/PixelStreaming2RTC/Private/Tests/DelegateTest.cpp
+++ b/Engine/Plugins/Media/PixelStreaming2/Source/PixelStreaming2RTC/Private/Tests/DelegateTest.cpp
@@ -234,8 +234,7 @@ namespace UE::PixelStreaming2
 			// This is also not being fired on linux or Mac so is disabled for now.
 			// auto& OnVideoTrackClosedNativeDelegate = CreateSingleDelegateTest<decltype(Delegates->OnVideoTrackClosedNative), FString, TOptional<FString>, bool>(DelegatesMap, Delegates->OnVideoTrackClosedNative, TEXT("OnVideoTrackClosedNative"))->Times(Exactly(Config.NumPlayers)).With(StreamerName, Any<FString>(), bIsRemote);
 			auto& OnAudioTrackOpenNativeDelegate = CreateSingleDelegateTest<decltype(Delegates->OnAudioTrackOpenNative), FString, TOptional<FString>, bool>(DelegatesMap, Delegates->OnAudioTrackOpenNative, TEXT("OnAudioTrackOpenNative"))->Times(Exactly(NumCalls)).With( StreamerName, Any<FString>(), bIsRemote);
-			// TODO (Eden.Harris) This should be called twice, once for local and remote
-			auto& OnAudioTrackClosedNativeDelegate = CreateSingleDelegateTest<decltype(Delegates->OnAudioTrackClosedNative), FString, TOptional<FString>, bool>(DelegatesMap, Delegates->OnAudioTrackClosedNative, TEXT("OnAudioTrackClosedNative"))->Times(Exactly(Config.NumPlayers)).With(StreamerName, Any<FString>(), bIsRemote);
+			auto& OnAudioTrackClosedNativeDelegate = CreateSingleDelegateTest<decltype(Delegates->OnAudioTrackClosedNative), FString, TOptional<FString>, bool>(DelegatesMap, Delegates->OnAudioTrackClosedNative, TEXT("OnAudioTrackClosedNative"))->Times(Exactly(NumCalls)).With(StreamerName, Any<FString>(), bIsRemote);
 			
 			if(Config.bIsBidirectional)
 			{
