[PS2] Dont display error log on last frame during stream.

Only display a error log if the frame before last did nto submit.

diff --git a/Engine/Plugins/Media/PixelStreaming2/Source/PixelStreaming2RTC/Private/EpicRtcVideoEncoder.cpp b/Engine/Plugins/Media/PixelStreaming2/Source/PixelStreaming2RTC/Private/EpicRtcVideoEncoder.cpp
--- a/Engine/Plugins/Media/PixelStreaming2/Source/PixelStreaming2RTC/Private/EpicRtcVideoEncoder.cpp
+++ b/Engine/Plugins/Media/PixelStreaming2/Source/PixelStreaming2RTC/Private/EpicRtcVideoEncoder.cpp
@@ -547,9 +547,24 @@ namespace UE::PixelStreaming2
 			EpicRtcVideoEncodedResult Result = VideoEncoderCallback->Encoded(EncodedFrame, CodecSpecificInfo);
 			UpdateFrameMetadataPostPacketization(*AdaptedLayer);
 
+			// It is not possible to know if the stream has ended here so a frame may be pushed despite the streaming ending.
+			// This causes the Result to return an error. 
+			// This section only prints if there was an error in the frame before current since that will be caused by an actual error.
 			if (Result._error)
 			{
-				UE_LOGFMT(LogPixelStreaming2RTC, Error, "PixelStreamingVideoEncoder: Failed to push encoded frame");
+				if (bDidLastEncodedFramePushFail)
+				{
+					UE_LOGFMT(LogPixelStreaming2RTC, Error, "PixelStreamingVideoEncoder: Failed to push previous and current encoded frame.");
+				}
+				else
+				{
+					UE_LOGFMT(LogPixelStreaming2RTC, Log, "PixelStreamingVideoEncoder: Failed to push encoded frame. This is expected when the stream is shutting down.");
+				}
+				bDidLastEncodedFramePushFail = true;
+			}
+			else
+			{
+				bDidLastEncodedFramePushFail = false;
 			}
 		}
 
diff --git a/Engine/Plugins/Media/PixelStreaming2/Source/PixelStreaming2RTC/Private/EpicRtcVideoEncoder.h b/Engine/Plugins/Media/PixelStreaming2/Source/PixelStreaming2RTC/Private/EpicRtcVideoEncoder.h
--- a/Engine/Plugins/Media/PixelStreaming2/Source/PixelStreaming2RTC/Private/EpicRtcVideoEncoder.h
+++ b/Engine/Plugins/Media/PixelStreaming2/Source/PixelStreaming2RTC/Private/EpicRtcVideoEncoder.h
@@ -53,6 +53,7 @@ namespace UE::PixelStreaming2
 
 		// File handle used to save the encoded output to
 		IFileHandle* FileHandle = nullptr;
+		bool		 bDidLastEncodedFramePushFail = false;
 
 	private:
 		bool	 LateInitEncoder();
