[PS2] QOL: Use PixelCapture for decode

This is part one of removing the need for PixelStreaming to to I420<->ARGB conversions. By using PixelCapture for decode, it is responsible for the conversion.

diff --git a/Engine/Plugins/Media/PixelStreaming2/Source/PixelStreaming2/Internal/VideoCapturer.h b/Engine/Plugins/Media/PixelStreaming2/Source/PixelStreaming2/Internal/VideoCapturer.h
--- a/Engine/Plugins/Media/PixelStreaming2/Source/PixelStreaming2/Internal/VideoCapturer.h
+++ b/Engine/Plugins/Media/PixelStreaming2/Source/PixelStreaming2/Internal/VideoCapturer.h
@@ -19,7 +19,7 @@ namespace UE::PixelStreaming2
 	{
 	public:
 		static TSharedPtr<FVideoCapturer> Create(TSharedPtr<FVideoProducer> VideoProducer = nullptr);
-		virtual ~FVideoCapturer() = default;
+		virtual ~FVideoCapturer();
 
 		// Begin IPixelCaptureCapturerSource Interface
 		virtual TSharedPtr<FPixelCaptureCapturer> CreateCapturer(int32 FinalFormat, float FinalScale) override;
@@ -35,6 +35,8 @@ namespace UE::PixelStreaming2
 
 		void ResetFrameCapturer();
 
+		void OnFrame(const IPixelCaptureInputFrame& InputFrame);
+
 		/**
 		 * This is broadcast each time a frame exits the adapt process. Used to synchronize framerates with input rates.
 		 * This should be called once per frame taking into consideration all the target formats and layers within the frame.
@@ -43,7 +45,7 @@ namespace UE::PixelStreaming2
 		FOnFrameCaptured OnFrameCaptured;
 
 		/**
-		 * Each time the frame capturer is created. This could be triggered by a multitude of situations such as cvar changes, frame 
+		 * Each time the frame capturer is created. This could be triggered by a multitude of situations such as cvar changes, frame
 		 * resizes, etc
 		 */
 		DECLARE_MULTICAST_DELEGATE(FOnFrameCapturerCreated);
@@ -62,12 +64,15 @@ namespace UE::PixelStreaming2
 		FDelegateHandle								 CaptureCompleteHandle;
 		FDelegateHandle								 FramePushedHandle;
 
+		FDelegateHandle SimulcastEnabledChangedHandle;
+		FDelegateHandle CaptureUseFenceChangedHandle;
+		FDelegateHandle UseMediaCaptureChangedHandle;
+
 		void CreateFrameCapturer();
 		void OnSimulcastEnabledChanged(IConsoleVariable* Var);
 		void OnCaptureUseFenceChanged(IConsoleVariable* Var);
 		void OnUseMediaCaptureChanged(IConsoleVariable* Var);
 		void OnCaptureComplete();
-		void OnFrame(const IPixelCaptureInputFrame& InputFrame);
 	};
 
 } // namespace UE::PixelStreaming2
diff --git a/Engine/Plugins/Media/PixelStreaming2/Source/PixelStreaming2/Private/VideoCapturer.cpp b/Engine/Plugins/Media/PixelStreaming2/Source/PixelStreaming2/Private/VideoCapturer.cpp
--- a/Engine/Plugins/Media/PixelStreaming2/Source/PixelStreaming2/Private/VideoCapturer.cpp
+++ b/Engine/Plugins/Media/PixelStreaming2/Source/PixelStreaming2/Private/VideoCapturer.cpp
@@ -18,28 +18,44 @@
 
 namespace UE::PixelStreaming2
 {
-
 	TSharedPtr<FVideoCapturer> FVideoCapturer::Create(TSharedPtr<FVideoProducer> InVideoProducer)
 	{
 		TSharedPtr<FVideoCapturer> VideoCapturer = TSharedPtr<FVideoCapturer>(new FVideoCapturer(InVideoProducer));
-		if (InVideoProducer)
+
+		return VideoCapturer;
+	}
+
+	FVideoCapturer::FVideoCapturer(TSharedPtr<FVideoProducer> VideoProducer)
+		: VideoProducer(VideoProducer)
+	{
+		CreateFrameCapturer();
+
+		if (VideoProducer)
 		{
-			VideoCapturer->FramePushedHandle = InVideoProducer->OnFramePushed.AddSP(VideoCapturer.ToSharedRef(), &FVideoCapturer::OnFrame);
+			FramePushedHandle = VideoProducer->OnFramePushed.AddRaw(this, &FVideoCapturer::OnFrame);
 		}
 
 		if (UPixelStreaming2PluginSettings::FDelegates* Delegates = UPixelStreaming2PluginSettings::Delegates())
 		{
-			Delegates->OnSimulcastEnabledChanged.AddSP(VideoCapturer.ToSharedRef(), &FVideoCapturer::OnSimulcastEnabledChanged);
-			Delegates->OnCaptureUseFenceChanged.AddSP(VideoCapturer.ToSharedRef(), &FVideoCapturer::OnCaptureUseFenceChanged);
-			Delegates->OnUseMediaCaptureChanged.AddSP(VideoCapturer.ToSharedRef(), &FVideoCapturer::OnUseMediaCaptureChanged);
+			SimulcastEnabledChangedHandle = Delegates->OnSimulcastEnabledChanged.AddRaw(this, &FVideoCapturer::OnSimulcastEnabledChanged);
+			CaptureUseFenceChangedHandle = Delegates->OnCaptureUseFenceChanged.AddRaw(this, &FVideoCapturer::OnCaptureUseFenceChanged);
+			UseMediaCaptureChangedHandle = Delegates->OnUseMediaCaptureChanged.AddRaw(this, &FVideoCapturer::OnUseMediaCaptureChanged);
 		}
-		return VideoCapturer;
 	}
 
-	FVideoCapturer::FVideoCapturer(TSharedPtr<FVideoProducer> VideoProducer)
-		: VideoProducer(VideoProducer)
+	FVideoCapturer::~FVideoCapturer()
 	{
-		CreateFrameCapturer();
+		if (VideoProducer)
+		{
+			VideoProducer->OnFramePushed.Remove(FramePushedHandle);
+		}
+
+		if (UPixelStreaming2PluginSettings::FDelegates* Delegates = UPixelStreaming2PluginSettings::Delegates())
+		{
+			Delegates->OnSimulcastEnabledChanged.Remove(SimulcastEnabledChangedHandle);
+			Delegates->OnCaptureUseFenceChanged.Remove(CaptureUseFenceChangedHandle);
+			Delegates->OnUseMediaCaptureChanged.Remove(UseMediaCaptureChangedHandle);
+		}
 	}
 
 	void FVideoCapturer::SetVideoProducer(TSharedPtr<FVideoProducer> InVideoProducer)
@@ -147,13 +163,13 @@ namespace UE::PixelStreaming2
 	TSharedPtr<FPixelCaptureCapturer> FVideoCapturer::CreateCapturer(int32 FinalFormat, float FinalScale)
 	{
 		OnFrameCapturerCreated.Broadcast();
-		
+
 		if (LastFrameType == PixelCaptureBufferFormat::FORMAT_RHI)
 		{
 			switch (FinalFormat)
 			{
 				case PixelCaptureBufferFormat::FORMAT_RHI:
-					if (VideoProducer->IsFrameAlreadyCopied() && FinalScale == 1.0)
+					if (VideoProducer && VideoProducer->IsFrameAlreadyCopied() && FinalScale == 1.0)
 					{
 						// If we only have a single layer (and it's scale is 1), we can use the no copy capturer
 						// as we know the output from the media capture will already be the correct format and scale
@@ -222,5 +238,4 @@ namespace UE::PixelStreaming2
 			return nullptr;
 		}
 	}
-
 } // namespace UE::PixelStreaming2
\ No newline at end of file
diff --git a/Engine/Plugins/Media/PixelStreaming2/Source/PixelStreaming2RTC/Internal/EpicRtcVideoSink.h b/Engine/Plugins/Media/PixelStreaming2/Source/PixelStreaming2RTC/Internal/EpicRtcVideoSink.h
--- a/Engine/Plugins/Media/PixelStreaming2/Source/PixelStreaming2RTC/Internal/EpicRtcVideoSink.h
+++ b/Engine/Plugins/Media/PixelStreaming2/Source/PixelStreaming2RTC/Internal/EpicRtcVideoSink.h
@@ -4,6 +4,7 @@
 
 #include "Containers/Set.h"
 #include "EpicRtcVideoTrack.h"
+#include "VideoCapturer.h"
 #include "VideoSink.h"
 #include "RendererInterface.h"
 
@@ -14,7 +15,7 @@ namespace UE::PixelStreaming2
 	/**
 	 * Video sink class that receives a frame from EpicRtc and passes the frame to all added consumers
 	 */
-	class PIXELSTREAMING2RTC_API FEpicRtcVideoSink : public FVideoSink, public FEpicRtcVideoTrack
+	class PIXELSTREAMING2RTC_API FEpicRtcVideoSink : public FVideoSink, public FEpicRtcVideoTrack, public TSharedFromThis<FEpicRtcVideoSink>
 	{
 	public:
 		static TSharedPtr<FEpicRtcVideoSink> Create(TRefCountPtr<EpicRtcVideoTrackInterface> InTrack);
@@ -26,10 +27,14 @@ namespace UE::PixelStreaming2
 	private:
 		FEpicRtcVideoSink(TRefCountPtr<EpicRtcVideoTrackInterface> InTrack);
 
+		void OnFrameCaptured();
+
 		FCriticalSection				  RenderSyncContext;
 		FPooledRenderTargetDesc			  RenderTargetDescriptor;
 		TRefCountPtr<IPooledRenderTarget> RenderTarget;
 		TArray<uint8_t>					  Buffer;
 		FTextureRHIRef					  SourceTexture;
+
+		TSharedPtr<FVideoCapturer> VideoCapturer;
 	};
 } // namespace UE::PixelStreaming2
\ No newline at end of file
diff --git a/Engine/Plugins/Media/PixelStreaming2/Source/PixelStreaming2RTC/Private/Blueprints/PixelStreaming2Peer.cpp b/Engine/Plugins/Media/PixelStreaming2/Source/PixelStreaming2RTC/Private/Blueprints/PixelStreaming2Peer.cpp
--- a/Engine/Plugins/Media/PixelStreaming2/Source/PixelStreaming2RTC/Private/Blueprints/PixelStreaming2Peer.cpp
+++ b/Engine/Plugins/Media/PixelStreaming2/Source/PixelStreaming2RTC/Private/Blueprints/PixelStreaming2Peer.cpp
@@ -95,7 +95,11 @@ bool UPixelStreaming2Peer::Disconnect(const FString& OptionalReason)
 
 	if (AudioSink)
 	{
-		AudioSink->RemoveAudioConsumer(this);
+		if (TSharedPtr<UE::PixelStreaming2::FEpicRtcAudioSink> SafeAudioSink = AudioSink; SafeAudioSink)
+		{
+			SafeAudioSink->RemoveAudioConsumer(this);
+		}
+
 		AudioSink = nullptr;
 		RemoteAudioTrack = 0;
 	}
@@ -107,7 +111,10 @@ bool UPixelStreaming2Peer::Disconnect(const FString& OptionalReason)
 	{
 		if (VideoConsumer)
 		{
-			VideoSink->RemoveVideoConsumer(TWeakPtrVariant<IPixelStreaming2VideoConsumer>(VideoConsumer.Get()));
+			if (TSharedPtr<UE::PixelStreaming2::FEpicRtcVideoSink> SafeVideoSink = VideoSink; SafeVideoSink)
+			{
+				SafeVideoSink->RemoveVideoConsumer(TWeakPtrVariant<IPixelStreaming2VideoConsumer>(VideoConsumer.Get()));
+			}
 		}
 		VideoSink = nullptr;
 		RemoteVideoTrack = 0;
@@ -302,7 +309,11 @@ void UPixelStreaming2Peer::OnAudioTrackUpdate(EpicRtcParticipantInterface* Parti
 		// We received a remote track. We should now generate audio from it
 		AudioSink = UE::PixelStreaming2::FEpicRtcAudioSink::Create(AudioTrack);
 		RemoteAudioTrack = reinterpret_cast<uintptr_t>(AudioTrack); // Keep track of which remote track we're receiving audio from as we only support one track
-		AudioSink->AddAudioConsumer(this);
+
+		if (TSharedPtr<UE::PixelStreaming2::FEpicRtcAudioSink> SafeAudioSink = AudioSink; SafeAudioSink)
+		{
+			SafeAudioSink->AddAudioConsumer(this);
+		}
 	}
 }
 
@@ -314,12 +325,18 @@ void UPixelStreaming2Peer::OnVideoTrackUpdate(EpicRtcParticipantInterface* Parti
 
 	if (VideoTrack->IsRemote())
 	{
-		// We received a remote track. We should now generate video from it
-		VideoSink = UE::PixelStreaming2::FEpicRtcVideoSink::Create(VideoTrack);
+		// We received a remote track. We should now create a sink to handle receiving the frames.
+		// NOTE: We pass in nullptr as the track because if we store the track on the sink, EpicRtc will be unable to destroy it
+		// and webrtc will try to flush remaining frames during session removal. 
+		VideoSink = UE::PixelStreaming2::FEpicRtcVideoSink::Create(nullptr);
 		RemoteVideoTrack = reinterpret_cast<uintptr_t>(VideoTrack); // Keep track of which remote track we're receiving video from as we only support display one track
 		if (VideoConsumer)
 		{
-			VideoSink->AddVideoConsumer(TWeakPtrVariant<IPixelStreaming2VideoConsumer>(VideoConsumer.Get()));
+			// As EpicRtc work can happen on different threads, we always need to call functions on a copy of the VideoSink to ensure it doesn't get deleted out from underneath us
+			if (TSharedPtr<UE::PixelStreaming2::FEpicRtcVideoSink> SafeVideoSink = VideoSink; SafeVideoSink)
+			{
+				SafeVideoSink->AddVideoConsumer(TWeakPtrVariant<IPixelStreaming2VideoConsumer>(VideoConsumer.Get()));
+			}
 		}
 	}
 }
@@ -349,7 +366,10 @@ void UPixelStreaming2Peer::OnAudioTrackMuted(EpicRtcAudioTrackInterface* AudioTr
 		return;
 	}
 
-	AudioSink->SetMuted((bool)bIsMuted);
+	if (TSharedPtr<UE::PixelStreaming2::FEpicRtcAudioSink> SafeAudioSink = AudioSink; SafeAudioSink)
+	{
+		SafeAudioSink->SetMuted(static_cast<bool>(bIsMuted));
+	}
 }
 
 void UPixelStreaming2Peer::OnAudioTrackFrame(EpicRtcAudioTrackInterface* AudioTrack, const EpicRtcAudioFrame& Frame)
@@ -359,7 +379,10 @@ void UPixelStreaming2Peer::OnAudioTrackFrame(EpicRtcAudioTrackInterface* AudioTr
 		return;
 	}
 
-	AudioSink->OnAudioData(Frame._data, Frame._length, Frame._format._numChannels, Frame._format._sampleRate);
+	if (TSharedPtr<UE::PixelStreaming2::FEpicRtcAudioSink> SafeAudioSink = AudioSink; SafeAudioSink)
+	{
+		SafeAudioSink->OnAudioData(Frame._data, Frame._length, Frame._format._numChannels, Frame._format._sampleRate);
+	}
 }
 
 void UPixelStreaming2Peer::OnAudioTrackRemoved(EpicRtcAudioTrackInterface* AudioTrack)
@@ -377,7 +400,10 @@ void UPixelStreaming2Peer::OnVideoTrackMuted(EpicRtcVideoTrackInterface* VideoTr
 		return;
 	}
 
-	VideoSink->SetMuted((bool)bIsMuted);
+	if (TSharedPtr<UE::PixelStreaming2::FEpicRtcVideoSink> SafeVideoSink = VideoSink; SafeVideoSink)
+	{
+		SafeVideoSink->SetMuted(static_cast<bool>(bIsMuted));
+	}
 }
 
 void UPixelStreaming2Peer::OnVideoTrackFrame(EpicRtcVideoTrackInterface* VideoTrack, const EpicRtcVideoFrame& Frame)
@@ -387,7 +413,10 @@ void UPixelStreaming2Peer::OnVideoTrackFrame(EpicRtcVideoTrackInterface* VideoTr
 		return;
 	}
 
-	VideoSink->OnEpicRtcFrame(Frame);
+	if (TSharedPtr<UE::PixelStreaming2::FEpicRtcVideoSink> SafeVideoSink = VideoSink; SafeVideoSink)
+	{
+		SafeVideoSink->OnEpicRtcFrame(Frame);
+	}
 }
 
 void UPixelStreaming2Peer::OnVideoTrackRemoved(EpicRtcVideoTrackInterface* VideoTrack)
diff --git a/Engine/Plugins/Media/PixelStreaming2/Source/PixelStreaming2RTC/Private/EpicRtcVideoBuffer.h b/Engine/Plugins/Media/PixelStreaming2/Source/PixelStreaming2RTC/Private/EpicRtcVideoBuffer.h
new file mode 100644
--- /dev/null
+++ b/Engine/Plugins/Media/PixelStreaming2/Source/PixelStreaming2RTC/Private/EpicRtcVideoBuffer.h
@@ -0,0 +1,44 @@
+// Copyright Epic Games, Inc. All Rights Reserved.
+
+#pragma once
+
+#include "PixelCaptureBufferFormat.h"
+
+#include "epic_rtc/core/video/video_buffer.h"
+#include "epic_rtc_helper/memory/ref_count_impl_helper.h"
+
+namespace UE::PixelStreaming2
+{
+	/**
+	 * The base VideoBuffer class that PixelStreaming2 uses. It partially implements
+	 * the EpicRtcVideoBufferInterface by always returning EpicRtcPixelFormat::Native to
+	 * signal to EpicRtc that it doesn't need to wrap this buffer in ANOTHER buffer before
+	 * passing to WebRTC. It provides the GetBufferFormat virtual method for use cases
+	 * to check the format of the buffer ie RHI or I420 etc
+	 */
+	class FEpicRtcVideoBuffer : public EpicRtcVideoBufferInterface
+	{
+	public:
+		FEpicRtcVideoBuffer() = default;
+
+		virtual ~FEpicRtcVideoBuffer() = default;
+
+	public:
+		virtual int32 GetBufferFormat()
+		{
+			return PixelCaptureBufferFormat::FORMAT_UNKNOWN;
+		}
+
+		// We always need to return Native so EpicRtc doesn't try to wrap
+		// these buffers in software wrappers
+		virtual EpicRtcPixelFormat GetFormat() override
+		{
+			return EpicRtcPixelFormat::Native;
+		}
+
+	public:
+		// Begin EpicRtcRefCountInterface
+		EPICRTC_REFCOUNT_INTERFACE_IN_PLACE
+		// End EpicRtcRefCountInterface
+	};
+} // namespace UE::PixelStreaming2
\ No newline at end of file
diff --git a/Engine/Plugins/Media/PixelStreaming2/Source/PixelStreaming2RTC/Private/EpicRtcVideoBufferI420.h b/Engine/Plugins/Media/PixelStreaming2/Source/PixelStreaming2RTC/Private/EpicRtcVideoBufferI420.h
--- a/Engine/Plugins/Media/PixelStreaming2/Source/PixelStreaming2RTC/Private/EpicRtcVideoBufferI420.h
+++ b/Engine/Plugins/Media/PixelStreaming2/Source/PixelStreaming2RTC/Private/EpicRtcVideoBufferI420.h
@@ -2,30 +2,33 @@
 
 #pragma once
 
+#include "EpicRtcVideoBuffer.h"
 #include "PixelCaptureBufferI420.h"
 
-#include "epic_rtc/core/video/video_buffer.h"
-#include "epic_rtc_helper/memory/ref_count_impl_helper.h"
-
 namespace UE::PixelStreaming2
 {
-	class FEpicRtcVideoBufferI420 : public EpicRtcVideoBufferInterface, public TRefCountingMixin<FEpicRtcVideoBufferI420>
+	class FEpicRtcVideoBufferI420 : public FEpicRtcVideoBuffer
 	{
 	public:
 		FEpicRtcVideoBufferI420(TSharedPtr<FPixelCaptureBufferI420> Buffer)
-			: Buffer(Buffer) {}
+			: Buffer(Buffer)
+		{
+		}
+
 		virtual ~FEpicRtcVideoBufferI420() = default;
 
 	public:
-		// Begin EpicRtcVideoBufferInterface
-		virtual void* GetData() override
+		// Begin FEpicRtcVideoBuffer
+		virtual int32 GetBufferFormat() override
 		{
-			return Buffer->GetMutableData();
+			return PixelCaptureBufferFormat::FORMAT_I420;
 		}
+		// End FEpicRtcVideoBuffer
 
-		virtual EpicRtcPixelFormat GetFormat() override
+		// Begin EpicRtcVideoBufferInterface
+		virtual void* GetData() override
 		{
-			return EpicRtcPixelFormat::I420;
+			return Buffer->GetMutableData();
 		}
 
 		virtual int GetWidth() override
@@ -39,10 +42,10 @@ namespace UE::PixelStreaming2
 		}
 		// End EpicRtcVideoBufferInterface
 
-	public:
-		// Begin EpicRtcRefCountInterface
-		EPICRTC_REFCOUNT_INTERFACE_IN_PLACE
-		// End EpicRtcRefCountInterface
+		TSharedPtr<FPixelCaptureBufferI420> GetBuffer()
+		{
+			return Buffer;
+		}
 
 	private:
 		TSharedPtr<FPixelCaptureBufferI420> Buffer;
diff --git a/Engine/Plugins/Media/PixelStreaming2/Source/PixelStreaming2RTC/Private/EpicRtcVideoBufferRHI.h b/Engine/Plugins/Media/PixelStreaming2/Source/PixelStreaming2RTC/Private/EpicRtcVideoBufferRHI.h
--- a/Engine/Plugins/Media/PixelStreaming2/Source/PixelStreaming2RTC/Private/EpicRtcVideoBufferRHI.h
+++ b/Engine/Plugins/Media/PixelStreaming2/Source/PixelStreaming2RTC/Private/EpicRtcVideoBufferRHI.h
@@ -2,28 +2,33 @@
 
 #pragma once
 
+#include "EpicRtcVideoBuffer.h"
 #include "Video/Resources/VideoResourceRHI.h"
 
-#include "epic_rtc/core/video/video_buffer.h"
-#include "epic_rtc_helper/memory/ref_count_impl_helper.h"
-
 namespace UE::PixelStreaming2
 {
-	class FEpicRtcVideoBufferRHI : public EpicRtcVideoBufferInterface
+	class FEpicRtcVideoBufferRHI : public FEpicRtcVideoBuffer
 	{
 	public:
 		FEpicRtcVideoBufferRHI(TSharedPtr<FVideoResourceRHI> VideoResourceRHI)
-			: VideoResourceRHI(VideoResourceRHI) 
+			: VideoResourceRHI(VideoResourceRHI)
 		{
 			VideoResourceRHI->SetUsing(true);
 		}
 
-		virtual ~FEpicRtcVideoBufferRHI() 
+		virtual ~FEpicRtcVideoBufferRHI()
 		{
 			VideoResourceRHI->SetUsing(false);
 		}
 
 	public:
+		// Begin FEpicRtcVideoBuffer
+		virtual int32 GetBufferFormat() override
+		{
+			return PixelCaptureBufferFormat::FORMAT_RHI;
+		}
+		// End FEpicRtcVideoBuffer
+
 		// Begin EpicRtcVideoBufferInterface
 		virtual void* GetData() override
 		{
@@ -31,11 +36,6 @@ namespace UE::PixelStreaming2
 			return nullptr;
 		}
 
-		virtual EpicRtcPixelFormat GetFormat() override
-		{
-			return EpicRtcPixelFormat::Native;
-		}
-
 		virtual int GetWidth() override
 		{
 			return VideoResourceRHI->GetDescriptor().Width;
@@ -47,7 +47,10 @@ namespace UE::PixelStreaming2
 		}
 		// End EpicRtcVideoBufferInterface
 
-		TSharedPtr<FVideoResourceRHI> GetVideoResource() { return VideoResourceRHI; }
+		TSharedPtr<FVideoResourceRHI> GetVideoResource()
+		{
+			return VideoResourceRHI;
+		}
 
 	public:
 		// Begin EpicRtcRefCountInterface
diff --git a/Engine/Plugins/Media/PixelStreaming2/Source/PixelStreaming2RTC/Private/EpicRtcVideoSink.cpp b/Engine/Plugins/Media/PixelStreaming2/Source/PixelStreaming2RTC/Private/EpicRtcVideoSink.cpp
--- a/Engine/Plugins/Media/PixelStreaming2/Source/PixelStreaming2RTC/Private/EpicRtcVideoSink.cpp
+++ b/Engine/Plugins/Media/PixelStreaming2/Source/PixelStreaming2RTC/Private/EpicRtcVideoSink.cpp
@@ -6,19 +6,31 @@
 #include "ColorConversion.h"
 #include "EpicRtcVideoBufferI420.h"
 #include "EpicRtcVideoBufferRHI.h"
+#include "IPixelCaptureOutputFrame.h"
+#include "Logging.h"
+#include "PixelCaptureInputFrameI420.h"
+#include "PixelCaptureInputFrameRHI.h"
+#include "PixelCaptureOutputFrameRHI.h"
 #include "PixelStreaming2Trace.h"
 #include "RenderTargetPool.h"
+#include "Stats.h"
 
 namespace UE::PixelStreaming2
 {
 	TSharedPtr<FEpicRtcVideoSink> FEpicRtcVideoSink::Create(TRefCountPtr<EpicRtcVideoTrackInterface> InTrack)
 	{
-		return TSharedPtr<FEpicRtcVideoSink>(new FEpicRtcVideoSink(InTrack));
+		TSharedPtr<FEpicRtcVideoSink> VideoSink = MakeShareable<FEpicRtcVideoSink>(new FEpicRtcVideoSink(InTrack));
+
+		VideoSink->VideoCapturer->OnFrameCaptured.AddSP(VideoSink.ToSharedRef(), &FEpicRtcVideoSink::OnFrameCaptured);
+
+		return VideoSink;
 	}
 
 	FEpicRtcVideoSink::FEpicRtcVideoSink(TRefCountPtr<EpicRtcVideoTrackInterface> InTrack)
 	{
 		Track = InTrack;
+
+		VideoCapturer = FVideoCapturer::Create(nullptr);
 	}
 
 	void FEpicRtcVideoSink::OnEpicRtcFrame(const EpicRtcVideoFrame& Frame)
@@ -33,7 +45,19 @@ namespace UE::PixelStreaming2
 		int32_t Width = Frame._buffer->GetWidth();
 		int32_t Height = Frame._buffer->GetHeight();
 
-		if (Frame._buffer->GetFormat() == EpicRtcPixelFormat::Native)
+		if (Frame._buffer->GetFormat() != EpicRtcPixelFormat::Native)
+		{
+			UE_LOGFMT(LogPixelStreaming2RTC, Error, "Received an EpicRtcVideoFrame that doesn't have a native buffer!");
+			return;
+		}
+
+		FEpicRtcVideoBuffer* const BaseFrameBuffer = static_cast<FEpicRtcVideoBuffer*>(Frame._buffer);
+		if (!BaseFrameBuffer)
+		{
+			return;
+		}
+
+		if (BaseFrameBuffer->GetBufferFormat() == PixelCaptureBufferFormat::FORMAT_RHI)
 		{
 			FEpicRtcVideoBufferRHI* const FrameBuffer = static_cast<FEpicRtcVideoBufferRHI*>(Frame._buffer);
 			if (FrameBuffer == nullptr)
@@ -42,89 +66,72 @@ namespace UE::PixelStreaming2
 			}
 
 			TSharedPtr<FVideoResourceRHI, ESPMode::ThreadSafe> VideoResource = FrameBuffer->GetVideoResource();
+
 			if (VideoResource->GetFormat() != EVideoFormat::BGRA)
 			{
 				VideoResource = VideoResource->TransformResource(FVideoDescriptor(EVideoFormat::BGRA, Width, Height));
 			}
 
-			auto& Raw = StaticCastSharedPtr<FVideoResourceRHI>(VideoResource)->GetRaw();
-			OnVideoData(Raw.Texture);
+			TWeakPtr<FEpicRtcVideoSink> WeakSink = AsWeak();
+			ENQUEUE_RENDER_COMMAND(CaptureDecodedFrameCommand)
+			([WeakSink, VideoResource](FRHICommandList& RHICmdList) {
+				if (TSharedPtr<FEpicRtcVideoSink> PinnedSink = WeakSink.Pin())
+				{
+					auto& Raw = VideoResource->GetRaw();
+					PinnedSink->VideoCapturer->OnFrame(FPixelCaptureInputFrameRHI(Raw.Texture));
+				}
+			});
 		}
-		else if (Frame._buffer->GetFormat() == EpicRtcPixelFormat::I420)
+		else if (BaseFrameBuffer->GetBufferFormat() == PixelCaptureBufferFormat::FORMAT_I420)
 		{
+			FEpicRtcVideoBufferI420* const FrameBuffer = static_cast<FEpicRtcVideoBufferI420*>(Frame._buffer);
+			if (FrameBuffer == nullptr)
 			{
-				FScopeLock	  Lock(&RenderSyncContext);
-				const int32_t Size = Width * Height * 4;
-				if (Size > Buffer.Num())
-				{
-					Buffer.SetNum(Size);
-				}
-
-				FEpicRtcVideoBufferI420* const FrameBuffer = static_cast<FEpicRtcVideoBufferI420*>(Frame._buffer);
-
-				int StrideY = FrameBuffer->GetWidth();
-				int StrideUV = (FrameBuffer->GetWidth() + 1) / 2;
-
-				int DataSizeY = StrideY * FrameBuffer->GetHeight();
-				int DataSizeUV = StrideUV * ((FrameBuffer->GetHeight() + 1) / 2);
-
-				uint8_t* DataY = static_cast<uint8_t*>(FrameBuffer->GetData());
-				uint8_t* DataU = static_cast<uint8_t*>(FrameBuffer->GetData()) + DataSizeY;
-				uint8_t* DataV = static_cast<uint8_t*>(FrameBuffer->GetData()) + DataSizeY + DataSizeUV;
-
-				ConvertI420ToArgb(
-					DataY, StrideY,
-					DataU, StrideUV,
-					DataV, StrideUV,
-					Buffer.GetData(), Width * 4,
-					FrameBuffer->GetWidth(), FrameBuffer->GetHeight());
+				return;
 			}
 
-			AsyncTask(ENamedThreads::GetRenderThread(), [Width, Height, this]() {
-				FScopeLock Lock(&RenderSyncContext);
+			TSharedPtr<FPixelCaptureBufferI420> I420Buffer = FrameBuffer->GetBuffer();
 
-				const FIntPoint			  FrameSize = FIntPoint(Width, Height);
-				FRHICommandListImmediate& RHICmdList = FRHICommandListExecutor::GetImmediateCommandList();
+			TWeakPtr<FEpicRtcVideoSink> WeakSink = AsWeak();
+			ENQUEUE_RENDER_COMMAND(CaptureDecodedFrameCommand)
+			([WeakSink, I420Buffer](FRHICommandList& RHICmdList) {
+				if (TSharedPtr<FEpicRtcVideoSink> PinnedSink = WeakSink.Pin())
+				{
+					PinnedSink->VideoCapturer->OnFrame(FPixelCaptureInputFrameI420(I420Buffer));
+				}
+			});
+		}
+	}
 
-				if (!RenderTargetDescriptor.IsValid() || RenderTargetDescriptor.GetSize() != FIntVector(FrameSize.X, FrameSize.Y, 0))
+	void FEpicRtcVideoSink::OnFrameCaptured()
+	{
+		TWeakPtr<FEpicRtcVideoSink> WeakSink = AsWeak();
+		ENQUEUE_RENDER_COMMAND(DisplayCapturedFrameCommand)
+		([WeakSink](FRHICommandList& RHICmdList) {
+			if (TSharedPtr<FEpicRtcVideoSink> PinnedSink = WeakSink.Pin())
+			{
+				TSharedPtr<IPixelCaptureOutputFrame> OutputFrame = PinnedSink->VideoCapturer->RequestFormat(PixelCaptureBufferFormat::FORMAT_RHI);
+				if (!OutputFrame)
 				{
-					// Create the RenderTarget descriptor
-					RenderTargetDescriptor = FPooledRenderTargetDesc::Create2DDesc(FrameSize,
-						PF_B8G8R8A8,
-						FClearValueBinding::None,
-						TexCreate_None,
-						TexCreate_RenderTargetable,
-						false);
-
-					// Update the shader resource for the 'SourceTexture'
-					FRHITextureCreateDesc RenderTargetTextureDesc =
-						FRHITextureCreateDesc::Create2D(TEXT(""), FrameSize.X, FrameSize.Y, PF_B8G8R8A8)
-							.SetClearValue(FClearValueBinding::None)
-#if PLATFORM_MAC
-							.SetFlags(ETextureCreateFlags::CPUReadback | ETextureCreateFlags::SRGB)
-							.SetInitialState(ERHIAccess::CPURead);
-#else
-							.SetFlags(ETextureCreateFlags::ShaderResource | TexCreate_RenderTargetable | ETextureCreateFlags::SRGB)
-							.SetInitialState(ERHIAccess::SRVMask);
-#endif
-
-					SourceTexture = RHICreateTexture(RenderTargetTextureDesc);
-
-					// Find a free target-able texture from the render pool
-					GRenderTargetPool.FindFreeElement(RHICmdList,
-						RenderTargetDescriptor,
-						RenderTarget,
-						TEXT("PIXELSTEAMINGPLAYER"));
+					return;
 				}
 
-				// Create the update region structure
-				const FUpdateTextureRegion2D Region(0, 0, 0, 0, FrameSize.X, FrameSize.Y);
+				TSharedPtr<FPixelCaptureOutputFrameRHI> RHIFrame = StaticCastSharedPtr<FPixelCaptureOutputFrameRHI>(OutputFrame);
+				if (!RHIFrame->GetFrameTexture())
+				{
+					return;
+				}
 
-				// Set the Pixel data of the webrtc Frame to the SourceTexture
-				RHIUpdateTexture2D(SourceTexture, 0, Region, FrameSize.X * 4, Buffer.GetData());
+				OutputFrame->Metadata.UseCount++;
+				if (OutputFrame->Metadata.UseCount == 1)
+				{
+					OutputFrame->Metadata.ProcessName = FString::Printf(TEXT("VideoSink %s"), *(OutputFrame->Metadata.ProcessName));
+				}
 
-				OnVideoData(SourceTexture);
-			});
-		}
+				FStats::Get()->AddFrameTimingStats(OutputFrame->Metadata);
+
+				PinnedSink->OnVideoData(RHIFrame->GetFrameTexture());
+			}
+		});
 	}
 } // namespace UE::PixelStreaming2
\ No newline at end of file
