[PS2] QOL: Make InputComponents storage thread safe

diff --git a/Engine/Plugins/Media/PixelStreaming2/Source/PixelStreaming2/Internal/PixelStreaming2Common.h b/Engine/Plugins/Media/PixelStreaming2/Source/PixelStreaming2/Internal/PixelStreaming2Common.h
--- a/Engine/Plugins/Media/PixelStreaming2/Source/PixelStreaming2/Internal/PixelStreaming2Common.h
+++ b/Engine/Plugins/Media/PixelStreaming2/Source/PixelStreaming2/Internal/PixelStreaming2Common.h
@@ -3,8 +3,9 @@
 #pragma once
 
 #include "Blueprints/PixelStreaming2InputComponent.h"
+#include "ThreadSafeMap.h"
 
 namespace UE::PixelStreaming2
 {
-	extern PIXELSTREAMING2_API TArray<UPixelStreaming2Input*> InputComponents;
+	extern PIXELSTREAMING2_API TThreadSafeMap<uintptr_t, UPixelStreaming2Input*> InputComponents;
 } // namespace UE::PixelStreaming2
diff --git a/Engine/Plugins/Media/PixelStreaming2/Source/PixelStreaming2/Private/Blueprints/PixelStreaming2InputComponent.cpp b/Engine/Plugins/Media/PixelStreaming2/Source/PixelStreaming2/Private/Blueprints/PixelStreaming2InputComponent.cpp
--- a/Engine/Plugins/Media/PixelStreaming2/Source/PixelStreaming2/Private/Blueprints/PixelStreaming2InputComponent.cpp
+++ b/Engine/Plugins/Media/PixelStreaming2/Source/PixelStreaming2/Private/Blueprints/PixelStreaming2InputComponent.cpp
@@ -22,14 +22,14 @@ void UPixelStreaming2Input::BeginPlay()
 {
 	Super::BeginPlay();
 
-	UE::PixelStreaming2::InputComponents.Add(this);
+	UE::PixelStreaming2::InputComponents.Add(reinterpret_cast<uintptr_t>(this), this);
 }
 
 void UPixelStreaming2Input::EndPlay(const EEndPlayReason::Type EndPlayReason)
 {
 	Super::EndPlay(EndPlayReason);
 
-	UE::PixelStreaming2::InputComponents.Remove(this);
+	UE::PixelStreaming2::InputComponents.Remove(reinterpret_cast<uintptr_t>(this));
 }
 
 void UPixelStreaming2Input::SendPixelStreaming2Response(const FString& Descriptor)
diff --git a/Engine/Plugins/Media/PixelStreaming2/Source/PixelStreaming2/Private/PixelStreaming2Module.cpp b/Engine/Plugins/Media/PixelStreaming2/Source/PixelStreaming2/Private/PixelStreaming2Module.cpp
--- a/Engine/Plugins/Media/PixelStreaming2/Source/PixelStreaming2/Private/PixelStreaming2Module.cpp
+++ b/Engine/Plugins/Media/PixelStreaming2/Source/PixelStreaming2/Private/PixelStreaming2Module.cpp
@@ -19,7 +19,7 @@
 
 namespace UE::PixelStreaming2
 {
-	TArray<UPixelStreaming2Input*> InputComponents = {};
+	TThreadSafeMap<uintptr_t, UPixelStreaming2Input*> InputComponents;
 
 	FPixelStreaming2Module* FPixelStreaming2Module::PixelStreaming2Module = nullptr;
 
diff --git a/Engine/Plugins/Media/PixelStreaming2/Source/PixelStreaming2RTC/Private/EpicRtcStreamer.cpp b/Engine/Plugins/Media/PixelStreaming2/Source/PixelStreaming2RTC/Private/EpicRtcStreamer.cpp
--- a/Engine/Plugins/Media/PixelStreaming2/Source/PixelStreaming2RTC/Private/EpicRtcStreamer.cpp
+++ b/Engine/Plugins/Media/PixelStreaming2/Source/PixelStreaming2RTC/Private/EpicRtcStreamer.cpp
@@ -339,10 +339,9 @@ namespace UE::PixelStreaming2
 		FString Descriptor = Res.Mid(1);
 
 		UE_LOGFMT(LogPixelStreaming2RTC, Verbose, "FEpicRtcStreamer[\"{0}\"]::OnUIInteraction({1})", StreamerId, Descriptor);
-		for (UPixelStreaming2Input* InputComponent : InputComponents)
-		{
-			InputComponent->OnInputEvent.Broadcast(Descriptor);
-		}
+		InputComponents.Apply([Descriptor](uintptr_t Key, UPixelStreaming2Input* Value) {
+			Value->OnInputEvent.Broadcast(Descriptor);
+		});
 	}
 
 	void FEpicRtcStreamer::OnSendMessage(FString MessageName, FMemoryReader Ar)
