[PS2] QOL: Unify PS2's EpicRtc track classes with a templated base class

diff --git a/Engine/Plugins/Media/PixelStreaming2/Source/PixelStreaming2RTC/Internal/EpicRtcAudioSink.h b/Engine/Plugins/Media/PixelStreaming2/Source/PixelStreaming2RTC/Internal/EpicRtcAudioSink.h
--- a/Engine/Plugins/Media/PixelStreaming2/Source/PixelStreaming2RTC/Internal/EpicRtcAudioSink.h
+++ b/Engine/Plugins/Media/PixelStreaming2/Source/PixelStreaming2RTC/Internal/EpicRtcAudioSink.h
@@ -3,12 +3,12 @@
 #pragma once
 
 #include "AudioSink.h"
-#include "EpicRtcAudioTrack.h"
+#include "EpicRtcTrack.h"
 
 namespace UE::PixelStreaming2
 {
 	// Collects audio coming in from EpicRtc and passes into into UE's audio system.
-	class PIXELSTREAMING2RTC_API FEpicRtcAudioSink : public FAudioSink, public FEpicRtcAudioTrack
+	class PIXELSTREAMING2RTC_API FEpicRtcAudioSink : public FAudioSink, public TEpicRtcTrack<EpicRtcAudioTrackInterface>
 	{
 	public:
 		static TSharedPtr<FEpicRtcAudioSink> Create(TRefCountPtr<EpicRtcAudioTrackInterface> InTrack);
diff --git a/Engine/Plugins/Media/PixelStreaming2/Source/PixelStreaming2RTC/Internal/EpicRtcAudioSource.h b/Engine/Plugins/Media/PixelStreaming2/Source/PixelStreaming2RTC/Internal/EpicRtcAudioSource.h
--- a/Engine/Plugins/Media/PixelStreaming2/Source/PixelStreaming2RTC/Internal/EpicRtcAudioSource.h
+++ b/Engine/Plugins/Media/PixelStreaming2/Source/PixelStreaming2RTC/Internal/EpicRtcAudioSource.h
@@ -3,14 +3,14 @@
 #pragma once
 
 #include "AudioSource.h"
-#include "EpicRtcAudioTrack.h"
+#include "EpicRtcTrack.h"
 #include "Templates/SharedPointer.h"
 
 #include "epic_rtc/core/audio/audio_track.h"
 
 namespace UE::PixelStreaming2
 {
-	class PIXELSTREAMING2RTC_API FEpicRtcAudioSource : public FAudioSource, public FEpicRtcAudioTrack
+	class PIXELSTREAMING2RTC_API FEpicRtcAudioSource : public FAudioSource, public TEpicRtcTrack<EpicRtcAudioTrackInterface>
 	{
 	public:
 		static TSharedPtr<FEpicRtcAudioSource> Create(TRefCountPtr<EpicRtcAudioTrackInterface> InTrack, TSharedPtr<class FEpicRtcAudioCapturer> InCapturer);
diff --git a/Engine/Plugins/Media/PixelStreaming2/Source/PixelStreaming2RTC/Internal/EpicRtcAudioTrack.h b/Engine/Plugins/Media/PixelStreaming2/Source/PixelStreaming2RTC/Internal/EpicRtcAudioTrack.h
deleted file mode 100644
diff --git a/Engine/Plugins/Media/PixelStreaming2/Source/PixelStreaming2RTC/Internal/EpicRtcDataTrack.h b/Engine/Plugins/Media/PixelStreaming2/Source/PixelStreaming2RTC/Internal/EpicRtcDataTrack.h
--- a/Engine/Plugins/Media/PixelStreaming2/Source/PixelStreaming2RTC/Internal/EpicRtcDataTrack.h
+++ b/Engine/Plugins/Media/PixelStreaming2/Source/PixelStreaming2RTC/Internal/EpicRtcDataTrack.h
@@ -2,11 +2,10 @@
 
 #pragma once
 
+#include "EpicRtcTrack.h"
 #include "IPixelStreaming2DataProtocol.h"
 #include "Templates/SharedPointer.h"
 
-#include "epic_rtc/core/data_track.h"
-
 namespace UE::PixelStreaming2
 {
 	template <typename T>
@@ -94,7 +93,7 @@ namespace UE::PixelStreaming2
 		}
 	};
 
-	class PIXELSTREAMING2RTC_API FEpicRtcDataTrack : public TSharedFromThis<FEpicRtcDataTrack>
+	class PIXELSTREAMING2RTC_API FEpicRtcDataTrack : public TEpicRtcTrack<EpicRtcDataTrackInterface>, public TSharedFromThis<FEpicRtcDataTrack>
 	{
 	public:
 		static TSharedPtr<FEpicRtcDataTrack> Create(TRefCountPtr<EpicRtcDataTrackInterface> InTrack, TWeakPtr<IPixelStreaming2DataProtocol> InDataProtocol);
@@ -133,11 +132,6 @@ namespace UE::PixelStreaming2
 		 */
 		bool SendArbitraryData(const FString& MessageType, const TArray64<uint8>& DataBytes) const;
 
-		/**
-		 * @return The id of the underlying EpicRtc data track.
-		 */
-		EpicRtcStringView GetId() const;
-
 		/**
 		 * @return The state of the underlying EpicRtc data track.
 		 */
@@ -156,8 +150,6 @@ namespace UE::PixelStreaming2
 		bool GetMessageId(const FString& MessageType, uint8& OutMessageId) const;
 
 	private:
-		TRefCountPtr<EpicRtcDataTrackInterface> Track;
-
 		/**
 		 * Track that is used for sending data with Consumer/Producer architecture.
 		 */
diff --git a/Engine/Plugins/Media/PixelStreaming2/Source/PixelStreaming2RTC/Internal/EpicRtcTrack.h b/Engine/Plugins/Media/PixelStreaming2/Source/PixelStreaming2RTC/Internal/EpicRtcTrack.h
new file mode 100644
--- /dev/null
+++ b/Engine/Plugins/Media/PixelStreaming2/Source/PixelStreaming2RTC/Internal/EpicRtcTrack.h
@@ -0,0 +1,37 @@
+// Copyright Epic Games, Inc. All Rights Reserved.
+
+#pragma once
+
+#include <type_traits>
+
+#include "Templates/RefCounting.h"
+
+#include "epic_rtc/core/audio/audio_track.h"
+#include "epic_rtc/core/data_track.h"
+#include "epic_rtc/core/video/video_track.h"
+
+namespace UE::PixelStreaming2
+{
+	template <typename TrackInterface UE_REQUIRES(std::is_same_v<TrackInterface, EpicRtcAudioTrackInterface> || std::is_same_v<TrackInterface, EpicRtcDataTrackInterface> || std::is_same_v<TrackInterface, EpicRtcVideoTrackInterface>)>
+	class TEpicRtcTrack
+	{
+	public:
+		TEpicRtcTrack(TrackInterface* Track)
+			: Track(Track)
+		{
+		}
+
+		virtual ~TEpicRtcTrack() = default;
+
+		/**
+		 * @return The id of the underlying EpicRtc track.
+		 */
+		EpicRtcStringView GetTrackId() const
+		{
+			return Track ? Track->GetId() : EpicRtcStringView{ ._ptr = nullptr, ._length = 0 };
+		}
+
+	protected:
+		TRefCountPtr<TrackInterface> Track;
+	};
+} // namespace UE::PixelStreaming2
\ No newline at end of file
diff --git a/Engine/Plugins/Media/PixelStreaming2/Source/PixelStreaming2RTC/Internal/EpicRtcVideoSink.h b/Engine/Plugins/Media/PixelStreaming2/Source/PixelStreaming2RTC/Internal/EpicRtcVideoSink.h
--- a/Engine/Plugins/Media/PixelStreaming2/Source/PixelStreaming2RTC/Internal/EpicRtcVideoSink.h
+++ b/Engine/Plugins/Media/PixelStreaming2/Source/PixelStreaming2RTC/Internal/EpicRtcVideoSink.h
@@ -3,7 +3,7 @@
 #pragma once
 
 #include "Containers/Set.h"
-#include "EpicRtcVideoTrack.h"
+#include "EpicRtcTrack.h"
 #include "VideoCapturer.h"
 #include "VideoSink.h"
 #include "RendererInterface.h"
@@ -15,7 +15,7 @@ namespace UE::PixelStreaming2
 	/**
 	 * Video sink class that receives a frame from EpicRtc and passes the frame to all added consumers
 	 */
-	class PIXELSTREAMING2RTC_API FEpicRtcVideoSink : public FVideoSink, public FEpicRtcVideoTrack, public TSharedFromThis<FEpicRtcVideoSink>
+	class PIXELSTREAMING2RTC_API FEpicRtcVideoSink : public FVideoSink, public TEpicRtcTrack<EpicRtcVideoTrackInterface>, public TSharedFromThis<FEpicRtcVideoSink>
 	{
 	public:
 		static TSharedPtr<FEpicRtcVideoSink> Create(TRefCountPtr<EpicRtcVideoTrackInterface> InTrack);
diff --git a/Engine/Plugins/Media/PixelStreaming2/Source/PixelStreaming2RTC/Internal/EpicRtcVideoSource.h b/Engine/Plugins/Media/PixelStreaming2/Source/PixelStreaming2RTC/Internal/EpicRtcVideoSource.h
--- a/Engine/Plugins/Media/PixelStreaming2/Source/PixelStreaming2RTC/Internal/EpicRtcVideoSource.h
+++ b/Engine/Plugins/Media/PixelStreaming2/Source/PixelStreaming2RTC/Internal/EpicRtcVideoSource.h
@@ -2,8 +2,8 @@
 
 #pragma once
 
+#include "EpicRtcTrack.h"
 #include "EpicRtcVideoCapturer.h"
-#include "EpicRtcVideoTrack.h"
 #include "VideoSource.h"
 
 namespace UE::PixelStreaming2
@@ -20,7 +20,7 @@ namespace UE::PixelStreaming2
 	 * will make sure that the adapter has valid output and if so will create a frame
 	 * for EpicRtc. Otherwise it will continue to wait until the next frame.
 	 */
-	class PIXELSTREAMING2RTC_API FEpicRtcVideoSource : public FVideoSource, public FEpicRtcVideoTrack
+	class PIXELSTREAMING2RTC_API FEpicRtcVideoSource : public FVideoSource, public TEpicRtcTrack<EpicRtcVideoTrackInterface>
 	{
 	public:
 		static TSharedPtr<FEpicRtcVideoSource> Create(TRefCountPtr<EpicRtcVideoTrackInterface> InVideoTrack, TSharedPtr<FEpicRtcVideoCapturer> InVideoCapturer, TSharedPtr<FVideoSourceGroup> InVideoSourceGroup);
diff --git a/Engine/Plugins/Media/PixelStreaming2/Source/PixelStreaming2RTC/Internal/EpicRtcVideoTrack.h b/Engine/Plugins/Media/PixelStreaming2/Source/PixelStreaming2RTC/Internal/EpicRtcVideoTrack.h
deleted file mode 100644
diff --git a/Engine/Plugins/Media/PixelStreaming2/Source/PixelStreaming2RTC/Private/EpicRtcAudioSink.cpp b/Engine/Plugins/Media/PixelStreaming2/Source/PixelStreaming2RTC/Private/EpicRtcAudioSink.cpp
--- a/Engine/Plugins/Media/PixelStreaming2/Source/PixelStreaming2RTC/Private/EpicRtcAudioSink.cpp
+++ b/Engine/Plugins/Media/PixelStreaming2/Source/PixelStreaming2RTC/Private/EpicRtcAudioSink.cpp
@@ -10,7 +10,7 @@ namespace UE::PixelStreaming2
 	}
 
 	FEpicRtcAudioSink::FEpicRtcAudioSink(TRefCountPtr<EpicRtcAudioTrackInterface> InTrack)
+		: TEpicRtcTrack(InTrack)
 	{
-		Track = InTrack;
 	}
 } // namespace UE::PixelStreaming2
\ No newline at end of file
diff --git a/Engine/Plugins/Media/PixelStreaming2/Source/PixelStreaming2RTC/Private/EpicRtcAudioSource.cpp b/Engine/Plugins/Media/PixelStreaming2/Source/PixelStreaming2RTC/Private/EpicRtcAudioSource.cpp
--- a/Engine/Plugins/Media/PixelStreaming2/Source/PixelStreaming2RTC/Private/EpicRtcAudioSource.cpp
+++ b/Engine/Plugins/Media/PixelStreaming2/Source/PixelStreaming2RTC/Private/EpicRtcAudioSource.cpp
@@ -18,8 +18,8 @@ namespace UE::PixelStreaming2
 	}
 
 	FEpicRtcAudioSource::FEpicRtcAudioSource(TRefCountPtr<EpicRtcAudioTrackInterface> InTrack)
+		: TEpicRtcTrack(InTrack)
 	{
-		Track = InTrack;
 	}
 
 	void FEpicRtcAudioSource::OnAudioBuffer(const int16_t* AudioData, int32 NumSamples, int32 NumChannels, const int32 SampleRate)
@@ -30,23 +30,22 @@ namespace UE::PixelStreaming2
 		}
 
 		const uint32_t NumFrames = NumSamples / NumChannels;
-		EpicRtcAudioFrame AudioFrame
-		{
+		// clang-format off
+		EpicRtcAudioFrame Frame = {
 			._data = const_cast<int16_t*>(AudioData),
 			._length = NumFrames,
 			._timestamp = 0,
-			._format =
-			{
-				._sampleRate = (uint32_t)SampleRate,
-				._numChannels = (uint32_t)NumChannels,
-				._parameters = nullptr
+			._format = {
+				._sampleRate = static_cast<uint32_t>(SampleRate),
+				._numChannels = static_cast<uint32_t>(NumChannels),
+				._parameters = nullptr 
 			}
 		};
+		// clang-format on
 
 		// Because UE handles all audio processing, we can bypass the ADM.
 		// This also has the added benefit of increasing audio quality
-		Track->PushFrame(AudioFrame, true);
+		Track->PushFrame(Frame, true);
 	}
 
 } // namespace UE::PixelStreaming2
-
diff --git a/Engine/Plugins/Media/PixelStreaming2/Source/PixelStreaming2RTC/Private/EpicRtcAudioTrack.cpp b/Engine/Plugins/Media/PixelStreaming2/Source/PixelStreaming2RTC/Private/EpicRtcAudioTrack.cpp
deleted file mode 100644
diff --git a/Engine/Plugins/Media/PixelStreaming2/Source/PixelStreaming2RTC/Private/EpicRtcDataTrack.cpp b/Engine/Plugins/Media/PixelStreaming2/Source/PixelStreaming2RTC/Private/EpicRtcDataTrack.cpp
--- a/Engine/Plugins/Media/PixelStreaming2/Source/PixelStreaming2RTC/Private/EpicRtcDataTrack.cpp
+++ b/Engine/Plugins/Media/PixelStreaming2/Source/PixelStreaming2RTC/Private/EpicRtcDataTrack.cpp
@@ -15,13 +15,13 @@ namespace UE::PixelStreaming2
 	}
 
 	FEpicRtcDataTrack::FEpicRtcDataTrack(TRefCountPtr<EpicRtcDataTrackInterface> InTrack, TWeakPtr<IPixelStreaming2DataProtocol> InDataProtocol)
-		: Track(InTrack)
+		: TEpicRtcTrack(InTrack)
 		, WeakDataProtocol(InDataProtocol)
 	{
 	}
 
 	FEpicRtcDataTrack::FEpicRtcDataTrack(TSharedPtr<FEpicRtcDataTrack> InTrack, TWeakPtr<IPixelStreaming2DataProtocol> InDataProtocol)
-		: Track(InTrack->SendTrack ? InTrack->SendTrack : InTrack->Track)
+		: TEpicRtcTrack(InTrack->SendTrack ? InTrack->SendTrack : InTrack->Track)
 		, WeakDataProtocol(InDataProtocol)
 	{
 	}
@@ -84,11 +84,6 @@ namespace UE::PixelStreaming2
 		return static_cast<bool>(SendResult);
 	}
 
-	EpicRtcStringView FEpicRtcDataTrack::GetId() const
-	{
-		return Track->GetId();
-	}
-
 	bool FEpicRtcDataTrack::SendArbitraryData(const FString& MessageType, const TArray64<uint8>& DataBytes) const
 	{
 		if (!IsActive())
diff --git a/Engine/Plugins/Media/PixelStreaming2/Source/PixelStreaming2RTC/Private/EpicRtcStreamer.cpp b/Engine/Plugins/Media/PixelStreaming2/Source/PixelStreaming2RTC/Private/EpicRtcStreamer.cpp
--- a/Engine/Plugins/Media/PixelStreaming2/Source/PixelStreaming2RTC/Private/EpicRtcStreamer.cpp
+++ b/Engine/Plugins/Media/PixelStreaming2/Source/PixelStreaming2RTC/Private/EpicRtcStreamer.cpp
@@ -22,7 +22,7 @@
 #include "Serialization/JsonSerializer.h"
 #include "Dom/JsonObject.h"
 #include "PixelStreaming2RTCModule.h"
-#include "EpicRtcAudioTrack.h"
+#include "EpicRtcTrack.h"
 #include "EpicRtcDataTrack.h"
 #include "IPixelStreaming2InputModule.h"
 #include "CoreGlobals.h"
diff --git a/Engine/Plugins/Media/PixelStreaming2/Source/PixelStreaming2RTC/Private/EpicRtcStreamer.h b/Engine/Plugins/Media/PixelStreaming2/Source/PixelStreaming2RTC/Private/EpicRtcStreamer.h
--- a/Engine/Plugins/Media/PixelStreaming2/Source/PixelStreaming2RTC/Private/EpicRtcStreamer.h
+++ b/Engine/Plugins/Media/PixelStreaming2/Source/PixelStreaming2RTC/Private/EpicRtcStreamer.h
@@ -148,7 +148,7 @@ namespace UE::PixelStreaming2
 				Participants->ApplyUntil([DataTrackId, &OutPlayerId](FString PlayerId, TSharedPtr<FPlayerContext> Participant) {
 					if (Participant->DataTrack)
 					{
-						FString TrackId = ToString(Participant->DataTrack->GetId());
+						FString TrackId = ToString(Participant->DataTrack->GetTrackId());
 						if (TrackId == DataTrackId)
 						{
 							OutPlayerId = PlayerId;
diff --git a/Engine/Plugins/Media/PixelStreaming2/Source/PixelStreaming2RTC/Private/EpicRtcVideoSink.cpp b/Engine/Plugins/Media/PixelStreaming2/Source/PixelStreaming2RTC/Private/EpicRtcVideoSink.cpp
--- a/Engine/Plugins/Media/PixelStreaming2/Source/PixelStreaming2RTC/Private/EpicRtcVideoSink.cpp
+++ b/Engine/Plugins/Media/PixelStreaming2/Source/PixelStreaming2RTC/Private/EpicRtcVideoSink.cpp
@@ -27,10 +27,9 @@ namespace UE::PixelStreaming2
 	}
 
 	FEpicRtcVideoSink::FEpicRtcVideoSink(TRefCountPtr<EpicRtcVideoTrackInterface> InTrack)
+		: TEpicRtcTrack(InTrack)
+		, VideoCapturer(FVideoCapturer::Create(nullptr))
 	{
-		Track = InTrack;
-
-		VideoCapturer = FVideoCapturer::Create(nullptr);
 	}
 
 	void FEpicRtcVideoSink::OnEpicRtcFrame(const EpicRtcVideoFrame& Frame)
diff --git a/Engine/Plugins/Media/PixelStreaming2/Source/PixelStreaming2RTC/Private/EpicRtcVideoSource.cpp b/Engine/Plugins/Media/PixelStreaming2/Source/PixelStreaming2RTC/Private/EpicRtcVideoSource.cpp
--- a/Engine/Plugins/Media/PixelStreaming2/Source/PixelStreaming2RTC/Private/EpicRtcVideoSource.cpp
+++ b/Engine/Plugins/Media/PixelStreaming2/Source/PixelStreaming2RTC/Private/EpicRtcVideoSource.cpp
@@ -18,9 +18,9 @@ namespace UE::PixelStreaming2
 	}
 
 	FEpicRtcVideoSource::FEpicRtcVideoSource(TRefCountPtr<EpicRtcVideoTrackInterface> InVideoTrack, TSharedPtr<FEpicRtcVideoCapturer> InVideoCapturer)
-		: VideoCapturer(InVideoCapturer)
+		: TEpicRtcTrack(InVideoTrack)
+		, VideoCapturer(InVideoCapturer)
 	{
-		Track = InVideoTrack;
 	}
 
 	TRefCountPtr<EpicRtcVideoTrackInterface> FEpicRtcVideoSource::GetVideoTrack()
diff --git a/Engine/Plugins/Media/PixelStreaming2/Source/PixelStreaming2RTC/Private/EpicRtcVideoTrack.cpp b/Engine/Plugins/Media/PixelStreaming2/Source/PixelStreaming2RTC/Private/EpicRtcVideoTrack.cpp
deleted file mode 100644
