[PS2] Fix match viewport changing the editor resolution, not the level editor viewport.

diff --git a/Engine/Plugins/Media/PixelStreaming2/Source/PixelStreaming2Editor/Private/PixelStreaming2EditorModule.cpp b/Engine/Plugins/Media/PixelStreaming2/Source/PixelStreaming2Editor/Private/PixelStreaming2EditorModule.cpp
--- a/Engine/Plugins/Media/PixelStreaming2/Source/PixelStreaming2Editor/Private/PixelStreaming2EditorModule.cpp
+++ b/Engine/Plugins/Media/PixelStreaming2/Source/PixelStreaming2Editor/Private/PixelStreaming2EditorModule.cpp
@@ -31,11 +31,39 @@
 #include "Widgets/Notifications/SNotificationList.h"
 #include "Logging.h"
 #include "IPixelStreaming2Module.h"
+#include "UtilsAsync.h"
 
 #define LOCTEXT_NAMESPACE "PixelStreaming2EditorModule"
 
 #define IMAGE_BRUSH_SVG(Style, RelativePath, ...) FSlateVectorImageBrush(Style.RootToContentDir(RelativePath, TEXT(".svg")), __VA_ARGS__)
 
+namespace
+{
+	void ResizeEditor(int Width, int Height)
+	{
+		TSharedPtr<SWindow> ParentWindow = IMainFrameModule::Get().GetParentWindow();
+		ParentWindow->Resize(FVector2D(Width, Height));
+		FSlateApplication::Get().OnSizeChanged(ParentWindow->GetNativeWindow().ToSharedRef(), Width, Height);
+		// Triggers the NullApplication to rebuild its DisplayMetrics with the new resolution and inform slate
+		// about the updated virtual desktop size
+		FSystemResolution::RequestResolutionChange(Width, Height, GSystemResolution.WindowMode);
+		IConsoleManager::Get().CallAllConsoleVariableSinks();
+	}
+	void ResizeViewport(int Width, int Height)
+	{
+		FLevelEditorModule& LevelEditorModule = FModuleManager::GetModuleChecked<FLevelEditorModule>("LevelEditor");
+		const TSharedPtr<SLevelViewport>  Viewport = LevelEditorModule.GetFirstActiveLevelViewport();
+		const FLevelEditorViewportClient* ViewportClient = Viewport ? &Viewport->GetLevelViewportClient() : nullptr;
+		const TSharedPtr<SEditorViewport> ViewportWidget = ViewportClient ? ViewportClient->GetEditorViewportWidget() : nullptr;
+		const TSharedPtr<FSceneViewport>  SceneViewport = ViewportClient ? ViewportWidget->GetSceneViewport() : nullptr;
+		if (SceneViewport)
+		{
+			SceneViewport->SetFixedViewportSize(Width, Height);
+		}
+	}
+
+}	// namespace
+
 namespace UE::EditorPixelStreaming2
 {
 	/**
@@ -211,7 +239,7 @@ namespace UE::EditorPixelStreaming2
 		// Add custom handle for { type: "Command", Resolution.Width: "1920", Resolution.Height: "1080" } when doing Editor streaming
 		// because we cannot resize the game viewport, but instead want to resize the parent window.
 		InputHandler->SetCommandHandler("Resolution.Width",
-			[](FString, FString Descriptor, FString WidthString) {
+			[InStreamType](FString, FString Descriptor, FString WidthString) {
 				bool	bSuccess;
 				FString HeightString;
 				UE::PixelStreaming2::ExtractJsonFromDescriptor(Descriptor, TEXT("Resolution.Height"), HeightString, bSuccess);
@@ -222,13 +250,14 @@ namespace UE::EditorPixelStreaming2
 					return;
 				}
 
-				TSharedPtr<SWindow> ParentWindow = IMainFrameModule::Get().GetParentWindow();
-				ParentWindow->Resize(FVector2D(Width, Height));
-				FSlateApplication::Get().OnSizeChanged(ParentWindow->GetNativeWindow().ToSharedRef(), Width, Height);
-				// Triggers the NullApplication to rebuild its DisplayMetrics with the new resolution and inform slate
-				// about the updated virtual desktop size
-				FSystemResolution::RequestResolutionChange(Width, Height, GSystemResolution.WindowMode);
-				IConsoleManager::Get().CallAllConsoleVariableSinks();
+				if (InStreamType == EPixelStreaming2EditorStreamTypes::LevelEditorViewport)
+				{
+					ResizeViewport(Width, Height);
+				}
+				else
+				{
+					ResizeEditor(Width, Height);
+				}
 			});
 
 		switch (InStreamType)
@@ -275,6 +304,16 @@ namespace UE::EditorPixelStreaming2
 			StartSignalling();
 		}
 
+		// If the level viewport has resized from the stream, this will reset it.
+		EditorStreamer->OnStreamingStopped().AddLambda([InStreamType](IPixelStreaming2Streamer*) {
+			if (InStreamType == EPixelStreaming2EditorStreamTypes::LevelEditorViewport)
+			{
+				DoOnGameThread([]() {
+					ResizeViewport(0, 0);
+				});
+			}
+		});
+
 		EditorStreamer->StartStreaming();
 	}
 
