[PS2] Fix: P2P connections receiving low quality frame when simulcast enabled

diff --git a/Engine/Plugins/Media/PixelStreaming2/Source/PixelStreaming2RTC/Private/EpicRtcStreamer.cpp b/Engine/Plugins/Media/PixelStreaming2/Source/PixelStreaming2RTC/Private/EpicRtcStreamer.cpp
--- a/Engine/Plugins/Media/PixelStreaming2/Source/PixelStreaming2RTC/Private/EpicRtcStreamer.cpp
+++ b/Engine/Plugins/Media/PixelStreaming2/Source/PixelStreaming2RTC/Private/EpicRtcStreamer.cpp
@@ -1327,35 +1327,48 @@ namespace UE::PixelStreaming2
 
 			int MaxFramerate = UPixelStreaming2PluginSettings::CVarWebRTCFps.GetValueOnAnyThread();
 
+			// WebRTC cannot set the bitrate outside the first initial biterate set by the VideoEncodingConfig.
+			// By setting a high value here, the real value can be set by SetConnectionRates which can fit within this range.
+			// Without this, changing the max bitrate at runtime will be capped at the initial max bitrate.
+			// SetConnectionRates below will set the real max bitrate.
+			constexpr uint32 InitialMinBitrate = 1'000;
+			constexpr uint32 InitialMaxBitrate = 1'000'000'000;
+
 			TArray<FPixelStreaming2SimulcastLayer> SimulcastParams = UE::PixelStreaming2::GetSimulcastParameters();
-			for (int i = 0; i < SimulcastParams.Num(); ++i)
+			if (IsSFU(ParticipantId) && SimulcastParams.Num() > 1)
 			{
-				const FPixelStreaming2SimulcastLayer& SpatialLayer = SimulcastParams[i];
-				const bool							  bIsSimulcast = SimulcastParams.Num() > 1;
-
-				if (bIsSimulcast)
+				for (int i = 0; i < SimulcastParams.Num(); ++i)
 				{
-					Rids.Add(FUtf8String("simulcast") + FUtf8String::FromInt(SimulcastParams.Num() - i));
-				}
+					const FPixelStreaming2SimulcastLayer& SpatialLayer = SimulcastParams[i];
 
-				// WebRTC cannot set the bitrate outside the first initial biterate set by the VideoEncodingConfig.
-				// By setting a high value here, the real value can be set by SetConnectionRates which can fit within this range.
-				// Without this, changing the max bitrate at runtime will be capped at the initial max bitrate.
-				// SetConnectionRates below will set the real max bitrate.
-				constexpr uint32 InitialMinBitrate = 1'000;
-				constexpr uint32 InitialMaxBitrate = 1'000'000'000;
+					Rids.Add(FUtf8String("simulcast") + FUtf8String::FromInt(SimulcastParams.Num() - i));
 
+					EpicRtcVideoEncodingConfig VideoEncodingConfig = {
+						// clang-format off
+						._rid = EpicRtcStringView{._ptr = (const char*)(*(Rids[i])), ._length = static_cast<uint64>(Rids[i].Len()) },
+						// clang-format on
+						._scaleResolutionDownBy = SpatialLayer.Scaling,
+						._scalabilityMode = static_cast<EpicRtcVideoScalabilityMode>(GetEnumFromCVar<EScalabilityMode>(UPixelStreaming2PluginSettings::CVarEncoderScalabilityMode)), // HACK if the Enums become un-aligned
+						._minBitrate = static_cast<uint32>(SpatialLayer.MinBitrate),
+						._maxBitrate = static_cast<uint32>(SpatialLayer.MaxBitrate),
+						._maxFrameRate = static_cast<uint8_t>(MaxFramerate)
+					};
+
+					VideoEncodingConfigs.Add(VideoEncodingConfig);
+				}
+			}
+			else
+			{
+				// Default video config for P2P or if simulcast disabled
 				EpicRtcVideoEncodingConfig VideoEncodingConfig = {
 					// clang-format off
 					// TODO (Migration): RTCP-7027 Maybe bug in EpicRtc? Setting an rid if there's only one config results in no video
-					._rid = SimulcastParams.Num() > 1
-						? EpicRtcStringView{._ptr = (const char*)(*(Rids[i])), ._length = static_cast<uint64>(Rids[i].Len()) }
-						: EpicRtcStringView{._ptr = nullptr, ._length = 0 },
+					._rid = EpicRtcStringView{._ptr = nullptr, ._length = 0 },
 					// clang-format on
-					._scaleResolutionDownBy = SpatialLayer.Scaling,
+					._scaleResolutionDownBy = 1.f,
 					._scalabilityMode = static_cast<EpicRtcVideoScalabilityMode>(GetEnumFromCVar<EScalabilityMode>(UPixelStreaming2PluginSettings::CVarEncoderScalabilityMode)), // HACK if the Enums become un-aligned
-					._minBitrate = bIsSimulcast ? static_cast<uint32>(SpatialLayer.MinBitrate) : InitialMinBitrate,
-					._maxBitrate = bIsSimulcast ? static_cast<uint32>(SpatialLayer.MaxBitrate) : InitialMaxBitrate,
+					._minBitrate = InitialMinBitrate,
+					._maxBitrate = InitialMaxBitrate,
 					._maxFrameRate = static_cast<uint8_t>(MaxFramerate)
 				};
 
