[PS2] Add: Log filtering to reduce EpicRtc log spam

diff --git a/Engine/Plugins/Media/PixelStreaming2/Source/PixelStreaming2RTC/Internal/EpicRtcLogging.h b/Engine/Plugins/Media/PixelStreaming2/Source/PixelStreaming2RTC/Internal/EpicRtcLogging.h
--- a/Engine/Plugins/Media/PixelStreaming2/Source/PixelStreaming2RTC/Internal/EpicRtcLogging.h
+++ b/Engine/Plugins/Media/PixelStreaming2/Source/PixelStreaming2RTC/Internal/EpicRtcLogging.h
@@ -2,6 +2,8 @@
 
 #pragma once
 
+#include "Delegates/IDelegateInstance.h"
+#include "Internationalization/Regex.h"
 #include "Logging/LogMacros.h"
 
 #include "epic_rtc/common/logging.h"
@@ -33,9 +35,55 @@ namespace UE::PixelStreaming2
 	static_assert(EpicRtcLogLevel::Trace == UnrealLogToEpicRtcCategoryMap[ELogVerbosity::Type::VeryVerbose]);
 	static_assert(EpicRtcLogLevel::Trace == UnrealLogToEpicRtcCategoryMap[ELogVerbosity::Type::All]);
 
+	/**
+	 * An interface providing functionality for manipulating EpicRtc logs.
+	 *
+	 * This includes completely removing log messages by checking IsFiltered()
+	 * or by redacting sensitive information using Censor()
+	 */
+	class PIXELSTREAMING2RTC_API ILogManipulator
+	{
+	public:
+		virtual ~ILogManipulator() = default;
+
+		/**
+		 * Checks if a log message should or shouldn't be logged.
+		 *
+		 * @return true is the log message is filtered and shouldn't be displayed
+		 */
+		virtual bool IsFiltered(ELogVerbosity::Type LogVerbosity, const FString& LogString) = 0;
+
+		virtual FString Censor(ELogVerbosity::Type LogVerbosity, const FString& LogString) = 0;
+	};
+
+	class PIXELSTREAMING2RTC_API FEpicRtcLogFilter : public ILogManipulator
+	{
+	public:
+		FEpicRtcLogFilter();
+		~FEpicRtcLogFilter();
+
+		virtual bool IsFiltered(ELogVerbosity::Type LogVerbosity, const FString& LogString) override;
+
+		virtual FString Censor(ELogVerbosity::Type LogVerbosity, const FString& LogString) override { return LogString; }
+
+	private:
+		void OnEpicRtcLogFilterChanged(IConsoleVariable* Var);
+		void ParseFilterString(const FString& LogFilterString);
+
+	private:
+		TArray<FRegexPattern> RegexPatterns;
+
+		FDelegateHandle EpicRtcLogFilterChangedHandle;
+	};
+
 	class PIXELSTREAMING2RTC_API FEpicRtcLogsRedirector : public EpicRtcLoggerInterface
 	{
 	public:
+		FEpicRtcLogsRedirector(TSharedPtr<ILogManipulator> LogManipulator = nullptr);
+
 		virtual void Log(const EpicRtcLogMessage& Message) override;
+
+	private:
+		TSharedPtr<ILogManipulator> LogManipulator;
 	};
 } // namespace UE::PixelStreaming2
\ No newline at end of file
diff --git a/Engine/Plugins/Media/PixelStreaming2/Source/PixelStreaming2RTC/Private/EpicRtcLogging.cpp b/Engine/Plugins/Media/PixelStreaming2/Source/PixelStreaming2RTC/Private/EpicRtcLogging.cpp
--- a/Engine/Plugins/Media/PixelStreaming2/Source/PixelStreaming2RTC/Private/EpicRtcLogging.cpp
+++ b/Engine/Plugins/Media/PixelStreaming2/Source/PixelStreaming2RTC/Private/EpicRtcLogging.cpp
@@ -2,11 +2,68 @@
 
 #include "EpicRtcLogging.h"
 
+#include "Logging/StructuredLog.h"
+#include "PixelStreaming2PluginSettings.h"
+
 DEFINE_LOG_CATEGORY(LogPixelStreaming2EpicRtc);
 DEFINE_LOG_CATEGORY(LogPixelStreaming2WebRtc);
 
 namespace UE::PixelStreaming2
 {
+	FEpicRtcLogFilter::FEpicRtcLogFilter()
+	{
+		ParseFilterString(UPixelStreaming2PluginSettings::CVarEpicRtcLogFilter.GetValueOnAnyThread());
+
+		if (UPixelStreaming2PluginSettings::FDelegates* Delegates = UPixelStreaming2PluginSettings::Delegates())
+		{
+			EpicRtcLogFilterChangedHandle = Delegates->OnEpicRtcLogFilterChanged.AddRaw(this, &FEpicRtcLogFilter::OnEpicRtcLogFilterChanged);
+		}
+	}
+
+	FEpicRtcLogFilter::~FEpicRtcLogFilter()
+	{
+		if (UPixelStreaming2PluginSettings::FDelegates* Delegates = UPixelStreaming2PluginSettings::Delegates())
+		{
+			Delegates->OnSimulcastEnabledChanged.Remove(EpicRtcLogFilterChangedHandle);
+		}
+	}
+
+	bool FEpicRtcLogFilter::IsFiltered(ELogVerbosity::Type LogVerbosity, const FString& LogString)
+	{
+		bool bFiltered = false;
+		for (const FRegexPattern& RegexPattern : RegexPatterns)
+		{
+			FRegexMatcher RegexMatcher(RegexPattern, LogString);
+			bFiltered |= RegexMatcher.FindNext();
+		}
+
+		return bFiltered;
+	}
+
+	void FEpicRtcLogFilter::OnEpicRtcLogFilterChanged(IConsoleVariable* Var)
+	{
+		ParseFilterString(Var->GetString());
+	}
+
+	void FEpicRtcLogFilter::ParseFilterString(const FString& LogFilterString)
+	{
+		TArray<FString> Filters;
+		LogFilterString.ParseIntoArray(Filters, TEXT("//"), true);
+
+		RegexPatterns.Empty();
+		for (const FString& Filter : Filters)
+		{
+			if (!Filter.IsEmpty())
+			{
+				RegexPatterns.Add(FRegexPattern(Filter));
+			}
+		}
+	}
+
+	FEpicRtcLogsRedirector::FEpicRtcLogsRedirector(TSharedPtr<ILogManipulator> LogManipulator)
+		: LogManipulator(LogManipulator)
+	{
+	}
 
 	void FEpicRtcLogsRedirector::Log(const EpicRtcLogMessage& Message)
 	{
@@ -27,39 +84,47 @@ namespace UE::PixelStreaming2
 		}
 		FString Msg{ (int32)Message._message._length, Message._message._ptr };
 
+	#define EPICRTC_LOG(LogLevel, LogString)                                                                                                                \
+		if (!LogManipulator || (LogManipulator && !LogManipulator->IsFiltered(ELogVerbosity::LogLevel, LogString)))                                         \
+		{                                                                                                                                                   \
+			UE_LOGFMT(LogPixelStreaming2EpicRtc, LogLevel, "{0}", LogManipulator ? LogManipulator->Censor(ELogVerbosity::LogLevel, LogString) : LogString); \
+		}
+
 		switch (Message._level)
 		{
 			case EpicRtcLogLevel::Trace:
 			{
-				UE_LOG(LogPixelStreaming2EpicRtc, VeryVerbose, TEXT("%s"), *Msg);
+				EPICRTC_LOG(VeryVerbose, Msg);
 				break;
 			}
 			case EpicRtcLogLevel::Debug:
 			{
-				UE_LOG(LogPixelStreaming2EpicRtc, Verbose, TEXT("%s"), *Msg);
+				EPICRTC_LOG(Verbose, Msg);
 				break;
 			}
 			case EpicRtcLogLevel::Info:
 			{
-				UE_LOG(LogPixelStreaming2EpicRtc, Log, TEXT("%s"), *Msg);
+				EPICRTC_LOG(Log, Msg);
 				break;
 			}
 			case EpicRtcLogLevel::Warning:
 			{
-				UE_LOG(LogPixelStreaming2EpicRtc, Warning, TEXT("%s"), *Msg);
+				EPICRTC_LOG(Warning, Msg);
 				break;
 			}
 			case EpicRtcLogLevel::Error:
 			{
-				UE_LOG(LogPixelStreaming2EpicRtc, Error, TEXT("%s"), *Msg);
+				EPICRTC_LOG(Error, Msg);
 				break;
 			}
 			case EpicRtcLogLevel::Critical:
 			{
-				UE_LOG(LogPixelStreaming2EpicRtc, Fatal, TEXT("%s"), *Msg);
+				EPICRTC_LOG(Fatal, Msg);
 				break;
 			}
 		}
+
+	#undef EPICRTC_LOG
 #endif
 	}
 
diff --git a/Engine/Plugins/Media/PixelStreaming2/Source/PixelStreaming2RTC/Private/PixelStreaming2RTCModule.cpp b/Engine/Plugins/Media/PixelStreaming2/Source/PixelStreaming2RTC/Private/PixelStreaming2RTCModule.cpp
--- a/Engine/Plugins/Media/PixelStreaming2/Source/PixelStreaming2RTC/Private/PixelStreaming2RTCModule.cpp
+++ b/Engine/Plugins/Media/PixelStreaming2/Source/PixelStreaming2RTC/Private/PixelStreaming2RTCModule.cpp
@@ -110,6 +110,9 @@ namespace UE::PixelStreaming2
 			return;
 		}
 
+		FString LogFilterString = UPixelStreaming2PluginSettings::CVarEpicRtcLogFilter.GetValueOnAnyThread() + TEXT("//\\bConference::Tick. Ticking audio too late\\b");
+		UPixelStreaming2PluginSettings::CVarEpicRtcLogFilter->Set(*LogFilterString, ECVF_SetByHotfix);
+
 		// By calling InitDefaultStreamer post engine init we can use pixel streaming in standalone editor mode
 		IPixelStreaming2Module::Get().OnReady().AddLambda([this](IPixelStreaming2Module& CoreModule) {
 			// Need to initialize after other modules have initialized such as NVCodec.
@@ -347,7 +350,7 @@ namespace UE::PixelStreaming2
 				._isGlobal = 0
 			},
 			._logging = {
-				._logger = new FEpicRtcLogsRedirector(),
+				._logger = new FEpicRtcLogsRedirector(MakeShared<FEpicRtcLogFilter>()),
 #if !NO_LOGGING // When building WITH_SHIPPING by default .GetVerbosity() does not exist
 				._level = UnrealLogToEpicRtcCategoryMap[LogPixelStreaming2EpicRtc.GetVerbosity()],
 				._levelWebRtc = UnrealLogToEpicRtcCategoryMap[LogPixelStreaming2WebRtc.GetVerbosity()]
diff --git a/Engine/Plugins/Media/PixelStreaming2/Source/PixelStreaming2Settings/Internal/PixelStreaming2PluginSettings.h b/Engine/Plugins/Media/PixelStreaming2/Source/PixelStreaming2Settings/Internal/PixelStreaming2PluginSettings.h
--- a/Engine/Plugins/Media/PixelStreaming2/Source/PixelStreaming2Settings/Internal/PixelStreaming2PluginSettings.h
+++ b/Engine/Plugins/Media/PixelStreaming2/Source/PixelStreaming2Settings/Internal/PixelStreaming2PluginSettings.h
@@ -81,6 +81,13 @@ public:
 		))
 	bool LogStats = false;
 
+	static TAutoConsoleVariable<FString> CVarEpicRtcLogFilter;
+	UPROPERTY(config, EditAnywhere, Category = "PixelStreaming", meta = (
+		DisplayName = "Log messages from EpicRtc to filter",
+		ToolTip = "Double forward slash (\"//\") separated list of regex patterns to filter from the EpicRtc logs (default: \"\")."
+		))
+	FString EpicRtcLogFilter = "";
+
 	static TAutoConsoleVariable<bool> CVarDisableLatencyTester;
 	UPROPERTY(config, EditAnywhere, Category = "PixelStreaming", meta = (
 		DisplayName = "Disable Latency Tester",
@@ -621,6 +628,9 @@ public:
 
 		DECLARE_TS_MULTICAST_DELEGATE_OneParam(FOnDecoupleFramerateChanged, IConsoleVariable*);
 		FOnDecoupleFramerateChanged OnDecoupleFramerateChanged;
+
+		DECLARE_TS_MULTICAST_DELEGATE_OneParam(FOnEpicRtcLogFilterChanged, IConsoleVariable*);
+		FOnEpicRtcLogFilterChanged OnEpicRtcLogFilterChanged;
 	};
 
 	static FDelegates* Delegates();
diff --git a/Engine/Plugins/Media/PixelStreaming2/Source/PixelStreaming2Settings/Private/PixelStreaming2PluginSettings.cpp b/Engine/Plugins/Media/PixelStreaming2/Source/PixelStreaming2Settings/Private/PixelStreaming2PluginSettings.cpp
--- a/Engine/Plugins/Media/PixelStreaming2/Source/PixelStreaming2Settings/Private/PixelStreaming2PluginSettings.cpp
+++ b/Engine/Plugins/Media/PixelStreaming2/Source/PixelStreaming2Settings/Private/PixelStreaming2PluginSettings.cpp
@@ -162,6 +162,7 @@ namespace
 // Map of Property Names to their commandline args as GetMetaData() is not avaliable in packaged projects
 static const TSet<TPair<FString, FString>> GetCmdArg = {
 	{ "PixelStreaming2.LogStats", "LogStats" },
+	{ "PixelStreaming2.EpicRtcLogFilter", "EpicRtcLogFilter" },
 	{ "PixelStreaming2.SendPlayerIdAsInteger", "SendPlayerIdAsInteger" },
 	{ "PixelStreaming2.DisableLatencyTester", "DisableLatencyTester" },
 	{ "PixelStreaming2.DecoupleFrameRate", "DecoupleFramerate" },
@@ -252,6 +253,13 @@ TAutoConsoleVariable<bool> UPixelStreaming2PluginSettings::CVarLogStats(
 	FConsoleVariableDelegate::CreateLambda([](IConsoleVariable* Var) { Delegates()->OnLogStatsChanged.Broadcast(Var); }),
 	ECVF_Default);
 
+TAutoConsoleVariable<FString> UPixelStreaming2PluginSettings::CVarEpicRtcLogFilter(
+	TEXT("PixelStreaming2.EpicRtcLogFilter"),
+	"",
+	TEXT("Double forward slash (\"//\") separated list of regex patterns to filter from the EpicRtc logs (default: \"\")."),
+	FConsoleVariableDelegate::CreateLambda([](IConsoleVariable* Var) { Delegates()->OnEpicRtcLogFilterChanged.Broadcast(Var); }),
+	ECVF_Default);
+
 TAutoConsoleVariable<bool> UPixelStreaming2PluginSettings::CVarDisableLatencyTester(
 	TEXT("PixelStreaming2.DisableLatencyTester"),
 	false,
