[PS2] Fix PS2 filtering audio log. Needs to check for 'to' and 'too'

diff --git a/Engine/Plugins/Media/PixelStreaming2/Source/PixelStreaming2RTC/Private/PixelStreaming2RTCModule.cpp b/Engine/Plugins/Media/PixelStreaming2/Source/PixelStreaming2RTC/Private/PixelStreaming2RTCModule.cpp
--- a/Engine/Plugins/Media/PixelStreaming2/Source/PixelStreaming2RTC/Private/PixelStreaming2RTCModule.cpp
+++ b/Engine/Plugins/Media/PixelStreaming2/Source/PixelStreaming2RTC/Private/PixelStreaming2RTCModule.cpp
@@ -110,7 +110,7 @@ namespace UE::PixelStreaming2
 			return;
 		}
 
-		FString LogFilterString = UPixelStreaming2PluginSettings::CVarEpicRtcLogFilter.GetValueOnAnyThread() + TEXT("//\\bConference::Tick. Ticking audio too late\\b");
+		FString LogFilterString = UPixelStreaming2PluginSettings::CVarEpicRtcLogFilter.GetValueOnAnyThread() + TEXT("//\\bConference::Tick. Ticking audio (?:too|to) late\\b");
 		UPixelStreaming2PluginSettings::CVarEpicRtcLogFilter->Set(*LogFilterString, ECVF_SetByHotfix);
 
 		// By calling InitDefaultStreamer post engine init we can use pixel streaming in standalone editor mode
