[PixelCapture, PS1, PS2] Fix: Memory leaks on DX11 due to fences not being removed. Only add if we're using experimental scheduling which means we will hit WaitForGPU

diff --git a/Engine/Plugins/Media/PixelCapture/Source/PixelCapture/Private/PixelCaptureCapturerMediaCapture.cpp b/Engine/Plugins/Media/PixelCapture/Source/PixelCapture/Private/PixelCaptureCapturerMediaCapture.cpp
--- a/Engine/Plugins/Media/PixelCapture/Source/PixelCapture/Private/PixelCaptureCapturerMediaCapture.cpp
+++ b/Engine/Plugins/Media/PixelCapture/Source/PixelCapture/Private/PixelCaptureCapturerMediaCapture.cpp
@@ -86,7 +86,10 @@ void UPixelCaptureMediaCapture::OnCustomCapture_RenderingThread(
 	FVector2D											   CropV)
 {
 	FGPUFenceRHIRef Fence = GDynamicRHI->RHICreateGPUFence(TEXT("UPixelCaptureMediaCapture Fence"));
-	Fences.PushLast(Fence);
+	if (UseExperimentalScheduling())
+	{
+		Fences.PushLast(Fence);
+	}
 
 	bool bRequiresFormatConversion = InSourceTexture->Desc.Format != OutputTexture->Desc.Format;
 	if (InSourceTexture->Desc.Format == OutputTexture->Desc.Format
diff --git a/Engine/Plugins/Media/PixelStreaming/Source/PixelStreaming/Private/PixelStreamingMediaIOCapture.cpp b/Engine/Plugins/Media/PixelStreaming/Source/PixelStreaming/Private/PixelStreamingMediaIOCapture.cpp
--- a/Engine/Plugins/Media/PixelStreaming/Source/PixelStreaming/Private/PixelStreamingMediaIOCapture.cpp
+++ b/Engine/Plugins/Media/PixelStreaming/Source/PixelStreaming/Private/PixelStreamingMediaIOCapture.cpp
@@ -54,7 +54,10 @@ void UPixelStreamingMediaIOCapture::OnCustomCapture_RenderingThread(
 	FVector2D											   CropV)
 {
 	FGPUFenceRHIRef Fence = GDynamicRHI->RHICreateGPUFence(TEXT("UPixelStreamingMediaIOCapture Fence"));
-	Fences.PushLast(Fence);
+	if (UseExperimentalScheduling())
+	{
+		Fences.PushLast(Fence);
+	}
 
 	bool bRequiresFormatConversion = InSourceTexture->Desc.Format != OutputTexture->Desc.Format;
 	if (InSourceTexture->Desc.Format == OutputTexture->Desc.Format && InSourceTexture->Desc.Extent.X == OutputTexture->Desc.Extent.X && InSourceTexture->Desc.Extent.Y == OutputTexture->Desc.Extent.Y)
diff --git a/Engine/Plugins/Media/PixelStreaming2/Source/PixelStreaming2/Private/PixelStreaming2MediaIOCapture.cpp b/Engine/Plugins/Media/PixelStreaming2/Source/PixelStreaming2/Private/PixelStreaming2MediaIOCapture.cpp
--- a/Engine/Plugins/Media/PixelStreaming2/Source/PixelStreaming2/Private/PixelStreaming2MediaIOCapture.cpp
+++ b/Engine/Plugins/Media/PixelStreaming2/Source/PixelStreaming2/Private/PixelStreaming2MediaIOCapture.cpp
@@ -31,7 +31,10 @@ void UPixelStreaming2MediaIOCapture::OnCustomCapture_RenderingThread(
 	FVector2D											   CropV)
 {
 	FGPUFenceRHIRef Fence = GDynamicRHI->RHICreateGPUFence(TEXT("UPixelStreaming2MediaIOCapture Fence"));
-	Fences.PushLast(Fence);
+	if (UseExperimentalScheduling())
+	{
+		Fences.PushLast(Fence);
+	}
 
 	bool bRequiresFormatConversion = InSourceTexture->Desc.Format != OutputTexture->Desc.Format;
 	if (InSourceTexture->Desc.Format == OutputTexture->Desc.Format
@@ -192,7 +195,7 @@ bool UPixelStreaming2MediaIOCapture::PostInitializeCaptureViewport(TSharedPtr<FS
 void UPixelStreaming2MediaIOCapture::ViewportResized(FViewport* Viewport, uint32 ResizeCode)
 {
 	bViewportResized = true;
-	
+
 	// If we have not captured a frame yet, we don't care about stopping the capture due to capture size mismatch
 	if (!CaptureResolution)
 	{
