[PS2] Fix: Race condition in Encoder/Decoder initializers when doing software fallback

diff --git a/Engine/Plugins/Media/PixelStreaming2/Source/PixelStreaming2RTC/Private/EpicRtcVideoDecoderInitializer.cpp b/Engine/Plugins/Media/PixelStreaming2/Source/PixelStreaming2RTC/Private/EpicRtcVideoDecoderInitializer.cpp
--- a/Engine/Plugins/Media/PixelStreaming2/Source/PixelStreaming2RTC/Private/EpicRtcVideoDecoderInitializer.cpp
+++ b/Engine/Plugins/Media/PixelStreaming2/Source/PixelStreaming2RTC/Private/EpicRtcVideoDecoderInitializer.cpp
@@ -92,7 +92,7 @@ namespace UE::PixelStreaming2
 	{
 		const TMap<EVideoCodec, TArray<TRefCountPtr<EpicRtcVideoCodecInfoInterface>>> SupportedCodecMap = CreateSupportedDecoderMap();
 
-		const EVideoCodec SelectedCodec = UE::PixelStreaming2::GetEnumFromCVar<EVideoCodec>(UPixelStreaming2PluginSettings::CVarEncoderCodec);
+		const EVideoCodec SelectedCodec = UE::PixelStreaming2::GetEnumFromString<EVideoCodec>(GSelectedCodec);
 		const bool		  bNegotiateCodecs = UPixelStreaming2PluginSettings::CVarWebRTCNegotiateCodecs.GetValueOnAnyThread();
 
 		// Build a list of supported codecs
diff --git a/Engine/Plugins/Media/PixelStreaming2/Source/PixelStreaming2RTC/Private/EpicRtcVideoEncoderInitializer.cpp b/Engine/Plugins/Media/PixelStreaming2/Source/PixelStreaming2RTC/Private/EpicRtcVideoEncoderInitializer.cpp
--- a/Engine/Plugins/Media/PixelStreaming2/Source/PixelStreaming2RTC/Private/EpicRtcVideoEncoderInitializer.cpp
+++ b/Engine/Plugins/Media/PixelStreaming2/Source/PixelStreaming2RTC/Private/EpicRtcVideoEncoderInitializer.cpp
@@ -97,7 +97,7 @@ namespace UE::PixelStreaming2
 		// static so we dont create the list every time this is called since the list will not change during runtime.
 		static const TMap<EVideoCodec, TArray<TRefCountPtr<EpicRtcVideoCodecInfoInterface>>> SupportedCodecMap = CreateSupportedEncoderMap();
 
-		EVideoCodec SelectedCodec = UE::PixelStreaming2::GetEnumFromCVar<EVideoCodec>(UPixelStreaming2PluginSettings::CVarEncoderCodec);
+		EVideoCodec SelectedCodec = UE::PixelStreaming2::GetEnumFromString<EVideoCodec>(GSelectedCodec);
 		const bool	bNegotiateCodecs = UPixelStreaming2PluginSettings::CVarWebRTCNegotiateCodecs.GetValueOnAnyThread();
 
 		// This array of support codecs gets built up and returned in our preference order
@@ -136,6 +136,7 @@ namespace UE::PixelStreaming2
 					}
 				});
 				// Also update our local SelectedCodec to reflect what the state will be
+				GSelectedCodec = TEXT("VP8");
 				SelectedCodec = EVideoCodec::VP8;
 				UE_LOG(LogPixelStreaming2RTC, Warning, TEXT("No more HW encoders available. Falling back to software encoding"));
 				DenyListedCodecs.Add(EVideoCodec::H264);
diff --git a/Engine/Plugins/Media/PixelStreaming2/Source/PixelStreaming2Settings/Internal/PixelStreaming2PluginSettings.h b/Engine/Plugins/Media/PixelStreaming2/Source/PixelStreaming2Settings/Internal/PixelStreaming2PluginSettings.h
--- a/Engine/Plugins/Media/PixelStreaming2/Source/PixelStreaming2Settings/Internal/PixelStreaming2PluginSettings.h
+++ b/Engine/Plugins/Media/PixelStreaming2/Source/PixelStreaming2Settings/Internal/PixelStreaming2PluginSettings.h
@@ -19,6 +19,11 @@
 
 namespace UE::PixelStreaming2
 {
+	// This variable exists to back CVarEncoderCodec and should only be used by the FEpicRtcVideo(Encoder/Decoder)Initializer. When doing software fallback,
+	// the encoder initializer will set post an async task to set the CVarEncoderCodec. However, we cannot be certain that this task will execute before
+	// the decoder initializer reads the cvar again. By having both use this variable, we can be certain that there will be no race conditions between them. 
+	extern PIXELSTREAMING2SETTINGS_API FString GSelectedCodec;
+
 	template <typename TEnumType>
 	TEnumType GetEnumFromCVar(const TAutoConsoleVariable<FString>& CVar)
 	{
@@ -27,6 +32,14 @@ namespace UE::PixelStreaming2
 		return static_cast<TEnumType>(StaticEnum<TEnumType>()->GetValueByIndex(OutEnum));
 	}
 
+	template <typename TEnumType>
+	TEnumType GetEnumFromString(const FString& EnumString)
+	{
+		uint64 OutEnum = StaticEnum<TEnumType>()->GetIndexByNameString(EnumString);
+		checkf(OutEnum != INDEX_NONE, TEXT("EnumString was not containing valid enum string"));
+		return static_cast<TEnumType>(StaticEnum<TEnumType>()->GetValueByIndex(OutEnum));
+	}
+
 	template <typename TEnumType>
 	FString GetCVarStringFromEnum(TEnumType Value)
 	{
diff --git a/Engine/Plugins/Media/PixelStreaming2/Source/PixelStreaming2Settings/Private/PixelStreaming2PluginSettings.cpp b/Engine/Plugins/Media/PixelStreaming2/Source/PixelStreaming2Settings/Private/PixelStreaming2PluginSettings.cpp
--- a/Engine/Plugins/Media/PixelStreaming2/Source/PixelStreaming2Settings/Private/PixelStreaming2PluginSettings.cpp
+++ b/Engine/Plugins/Media/PixelStreaming2/Source/PixelStreaming2Settings/Private/PixelStreaming2PluginSettings.cpp
@@ -403,11 +403,13 @@ TAutoConsoleVariable<bool> UPixelStreaming2PluginSettings::CVarEncoderEnableSimu
 	FConsoleVariableDelegate::CreateLambda([](IConsoleVariable* Var) { VerifyCVarVideoSettings(nullptr); Delegates()->OnSimulcastEnabledChanged.Broadcast(Var); }),
 	ECVF_Default);
 
+FString UE::PixelStreaming2::GSelectedCodec = TEXT("H264");
+
 TAutoConsoleVariable<FString> UPixelStreaming2PluginSettings::CVarEncoderCodec(
 	TEXT("PixelStreaming2.Encoder.Codec"),
-	"H264",
+	UE::PixelStreaming2::GSelectedCodec,
 	TEXT("PixelStreaming default encoder codec. Supported values are: `H264`, `VP8`, `VP9` or `AV1`"),
-	FConsoleVariableDelegate::CreateStatic(&VerifyCVarVideoSettings),
+	FConsoleVariableDelegate::CreateLambda([](IConsoleVariable* Var) { UE::PixelStreaming2::GSelectedCodec = Var->GetString(); VerifyCVarVideoSettings(nullptr); }),
 	ECVF_Default);
 
 TAutoConsoleVariable<FString> UPixelStreaming2PluginSettings::CVarEncoderScalabilityMode(
