[PS2] Fix: CVar only project settings not being parsed from command line

diff --git a/Engine/Plugins/Media/PixelStreaming2/Source/PixelStreaming2Settings/Private/PixelStreaming2PluginSettings.cpp b/Engine/Plugins/Media/PixelStreaming2/Source/PixelStreaming2Settings/Private/PixelStreaming2PluginSettings.cpp
--- a/Engine/Plugins/Media/PixelStreaming2/Source/PixelStreaming2Settings/Private/PixelStreaming2PluginSettings.cpp
+++ b/Engine/Plugins/Media/PixelStreaming2/Source/PixelStreaming2Settings/Private/PixelStreaming2PluginSettings.cpp
@@ -1348,6 +1348,10 @@ void UPixelStreaming2PluginSettings::PostInitProperties()
 
 	// Handle parsing of legacy command line args (such as -PixelStreamingUrl) after .ini and new commandline args.
 	ParseLegacyCommandlineArgs();
+
+	// These cvars don't have matching properties so need to be manually parsed
+	ParseLegacyCommandLineOption(*ConsoleVariableToCommandArgParam(TEXT("PixelStreaming2.Encoder.DumpDebugFrames")), CVarEncoderDebugDumpFrame);
+	ParseLegacyCommandLineOption(*ConsoleVariableToCommandArgParam(TEXT("PixelStreaming2.DumpDebugAudio")), CVarDebugDumpAudio);
 }
 
 UPixelStreaming2PluginSettings::FDelegates* UPixelStreaming2PluginSettings::Delegates()
