[PS1, PS2] Update get_ps_servers scripts to get 5.6 branch.

diff --git a/Engine/Plugins/Media/PixelStreaming/Resources/WebServers/get_ps_servers.bat b/Engine/Plugins/Media/PixelStreaming/Resources/WebServers/get_ps_servers.bat
--- a/Engine/Plugins/Media/PixelStreaming/Resources/WebServers/get_ps_servers.bat
+++ b/Engine/Plugins/Media/PixelStreaming/Resources/WebServers/get_ps_servers.bat
@@ -92,11 +92,15 @@ if DEFINED UEVersion (
     SET PSInfraTagOrBranch=UE5.5
     SET IsTag=0
   )
+  if "%UEVersion%"=="5.6" (
+    SET PSInfraTagOrBranch=UE5.6
+    SET IsTag=0
+  )
 )
 
 @Rem If no arguments select a specific version, fetch the appropriate default
 if NOT DEFINED PSInfraTagOrBranch (
-    SET PSInfraTagOrBranch=UE5.5
+    SET PSInfraTagOrBranch=UE5.6
     SET IsTag=0
 )
 echo Tag or branch: !PSInfraTagOrBranch!
diff --git a/Engine/Plugins/Media/PixelStreaming/Resources/WebServers/get_ps_servers.sh b/Engine/Plugins/Media/PixelStreaming/Resources/WebServers/get_ps_servers.sh
--- a/Engine/Plugins/Media/PixelStreaming/Resources/WebServers/get_ps_servers.sh
+++ b/Engine/Plugins/Media/PixelStreaming/Resources/WebServers/get_ps_servers.sh
@@ -26,7 +26,7 @@ print_help() {
 # Set all default variables (e.g. # Name and version of ps-infra that we are downloading)
 PSInfraOrg=EpicGamesExt
 PSInfraRepo=PixelStreamingInfrastructure
-PSInfraTagOrBranch=UE5.5
+PSInfraTagOrBranch=UE5.6
 RefType=heads
 IsTag=0
 ReleaseUrlBase=[REDACTED_PROFILE_URL]/PixelStreamingInfrastructure/releases/download
@@ -91,12 +91,17 @@ then
     PSInfraTagOrBranch=UE5.5
     IsTag=0
   fi
+  if [ "$UEVersion" = "5.6" ]
+  then
+    PSInfraTagOrBranch=UE5.6
+    IsTag=0
+  fi
 fi
 
 # If no arguments select a specific version, fetch the appropriate default
 if [ -z "$PSInfraTagOrBranch" ]
 then
-  PSInfraTagOrBranch=UE5.5
+  PSInfraTagOrBranch=UE5.6
   IsTag=0
 fi
 echo "Tag or branch: $PSInfraTagOrBranch"
diff --git a/Engine/Plugins/Media/PixelStreaming2/Resources/WebServers/get_ps_servers.bat b/Engine/Plugins/Media/PixelStreaming2/Resources/WebServers/get_ps_servers.bat
--- a/Engine/Plugins/Media/PixelStreaming2/Resources/WebServers/get_ps_servers.bat
+++ b/Engine/Plugins/Media/PixelStreaming2/Resources/WebServers/get_ps_servers.bat
@@ -92,11 +92,15 @@ if DEFINED UEVersion (
     SET PSInfraTagOrBranch=UE5.5
     SET IsTag=0
   )
+  if "%UEVersion%"=="5.6" (
+    SET PSInfraTagOrBranch=UE5.6
+    SET IsTag=0
+  )
 )
 
 @Rem If no arguments select a specific version, fetch the appropriate default
 if NOT DEFINED PSInfraTagOrBranch (
-    SET PSInfraTagOrBranch=UE5.5
+    SET PSInfraTagOrBranch=UE5.6
     SET IsTag=0
 )
 echo Tag or branch: !PSInfraTagOrBranch!
diff --git a/Engine/Plugins/Media/PixelStreaming2/Resources/WebServers/get_ps_servers.sh b/Engine/Plugins/Media/PixelStreaming2/Resources/WebServers/get_ps_servers.sh
--- a/Engine/Plugins/Media/PixelStreaming2/Resources/WebServers/get_ps_servers.sh
+++ b/Engine/Plugins/Media/PixelStreaming2/Resources/WebServers/get_ps_servers.sh
@@ -26,7 +26,7 @@ print_help() {
 # Set all default variables (e.g. # Name and version of ps-infra that we are downloading)
 PSInfraOrg=EpicGamesExt
 PSInfraRepo=PixelStreamingInfrastructure
-PSInfraTagOrBranch=UE5.5
+PSInfraTagOrBranch=UE5.6
 RefType=heads
 IsTag=0
 ReleaseUrlBase=[REDACTED_PROFILE_URL]/PixelStreamingInfrastructure/releases/download
@@ -91,12 +91,17 @@ then
     PSInfraTagOrBranch=UE5.5
     IsTag=0
   fi
+  if [ "$UEVersion" = "5.6" ]
+  then
+    PSInfraTagOrBranch=UE5.6
+    IsTag=0
+  fi
 fi
 
 # If no arguments select a specific version, fetch the appropriate default
 if [ -z "$PSInfraTagOrBranch" ]
 then
-  PSInfraTagOrBranch=UE5.5
+  PSInfraTagOrBranch=UE5.6
   IsTag=0
 fi
 echo "Tag or branch: $PSInfraTagOrBranch"
