[PS2] Enable Software fallback test in PS2. Fixed video/audio on closed delegates not firing. Fixed Delegate tests.

diff --git a/Engine/Plugins/Media/PixelStreaming2/Source/PixelStreaming2RTC/Internal/AtomicOnceFlag.h b/Engine/Plugins/Media/PixelStreaming2/Source/PixelStreaming2RTC/Internal/AtomicOnceFlag.h
new file mode 100644
--- /dev/null
+++ b/Engine/Plugins/Media/PixelStreaming2/Source/PixelStreaming2RTC/Internal/AtomicOnceFlag.h
@@ -0,0 +1,19 @@
+// Copyright Epic Games, Inc. All Rights Reserved.
+
+#pragma once
+
+#include "CoreMinimal.h"
+
+#include <atomic>
+
+namespace UE::PixelStreaming2
+{
+	class PIXELSTREAMING2RTC_API FAtomicOnceFlag
+	{
+	public:
+		// Sets flag to true and returns if the flag had fired before
+		inline bool Trigger() { return bFlag.exchange(true); };
+	private:
+		std::atomic<bool> bFlag = false;
+	};
+} // namespace UE::PixelStreaming2
diff --git a/Engine/Plugins/Media/PixelStreaming2/Source/PixelStreaming2RTC/Internal/EpicRtcAudioSink.h b/Engine/Plugins/Media/PixelStreaming2/Source/PixelStreaming2RTC/Internal/EpicRtcAudioSink.h
--- a/Engine/Plugins/Media/PixelStreaming2/Source/PixelStreaming2RTC/Internal/EpicRtcAudioSink.h
+++ b/Engine/Plugins/Media/PixelStreaming2/Source/PixelStreaming2RTC/Internal/EpicRtcAudioSink.h
@@ -2,6 +2,7 @@
 
 #pragma once
 
+#include "AtomicOnceFlag.h"
 #include "AudioSink.h"
 #include "EpicRtcTrack.h"
 
@@ -13,8 +14,11 @@ namespace UE::PixelStreaming2
 	public:
 		static TSharedPtr<FEpicRtcAudioSink> Create(TRefCountPtr<EpicRtcAudioTrackInterface> InTrack);
 		virtual ~FEpicRtcAudioSink() = default;
-
+		
+		inline bool TriggerDisconnectFlag() { return bDisconnectFlag.Trigger(); };
 	private:
 		FEpicRtcAudioSink(TRefCountPtr<EpicRtcAudioTrackInterface> InTrack);
+
+		FAtomicOnceFlag bDisconnectFlag;
 	};
 } // namespace UE::PixelStreaming2
\ No newline at end of file
diff --git a/Engine/Plugins/Media/PixelStreaming2/Source/PixelStreaming2RTC/Internal/EpicRtcAudioSource.h b/Engine/Plugins/Media/PixelStreaming2/Source/PixelStreaming2RTC/Internal/EpicRtcAudioSource.h
--- a/Engine/Plugins/Media/PixelStreaming2/Source/PixelStreaming2RTC/Internal/EpicRtcAudioSource.h
+++ b/Engine/Plugins/Media/PixelStreaming2/Source/PixelStreaming2RTC/Internal/EpicRtcAudioSource.h
@@ -2,6 +2,7 @@
 
 #pragma once
 
+#include "AtomicOnceFlag.h"
 #include "AudioSource.h"
 #include "EpicRtcTrack.h"
 #include "Templates/SharedPointer.h"
@@ -17,7 +18,10 @@ namespace UE::PixelStreaming2
 
 		virtual void OnAudioBuffer(const int16_t* AudioData, int32 NumSamples, int32 NumChannels, const int32 SampleRate);
 
+		inline bool TriggerDisconnectFlag() { return bDisconnectFlag.Trigger(); };
 	private:
 		FEpicRtcAudioSource(TRefCountPtr<EpicRtcAudioTrackInterface> InTrack);
+
+		FAtomicOnceFlag bDisconnectFlag;
 	};
 } // namespace UE::PixelStreaming2
\ No newline at end of file
diff --git a/Engine/Plugins/Media/PixelStreaming2/Source/PixelStreaming2RTC/Internal/EpicRtcVideoSink.h b/Engine/Plugins/Media/PixelStreaming2/Source/PixelStreaming2RTC/Internal/EpicRtcVideoSink.h
--- a/Engine/Plugins/Media/PixelStreaming2/Source/PixelStreaming2RTC/Internal/EpicRtcVideoSink.h
+++ b/Engine/Plugins/Media/PixelStreaming2/Source/PixelStreaming2RTC/Internal/EpicRtcVideoSink.h
@@ -2,6 +2,7 @@
 
 #pragma once
 
+#include "AtomicOnceFlag.h"
 #include "Containers/Set.h"
 #include "EpicRtcTrack.h"
 #include "VideoCapturer.h"
@@ -24,6 +25,7 @@ namespace UE::PixelStreaming2
 
 		void OnEpicRtcFrame(const EpicRtcVideoFrame& Frame);
 
+		inline bool TriggerDisconnectFlag() { return bDisconnectFlag.Trigger(); };
 	private:
 		FEpicRtcVideoSink(TRefCountPtr<EpicRtcVideoTrackInterface> InTrack);
 
@@ -36,5 +38,6 @@ namespace UE::PixelStreaming2
 		FTextureRHIRef					  SourceTexture;
 
 		TSharedPtr<FVideoCapturer> VideoCapturer;
+		FAtomicOnceFlag			   bDisconnectFlag;
 	};
 } // namespace UE::PixelStreaming2
\ No newline at end of file
diff --git a/Engine/Plugins/Media/PixelStreaming2/Source/PixelStreaming2RTC/Internal/EpicRtcVideoSource.h b/Engine/Plugins/Media/PixelStreaming2/Source/PixelStreaming2RTC/Internal/EpicRtcVideoSource.h
--- a/Engine/Plugins/Media/PixelStreaming2/Source/PixelStreaming2RTC/Internal/EpicRtcVideoSource.h
+++ b/Engine/Plugins/Media/PixelStreaming2/Source/PixelStreaming2RTC/Internal/EpicRtcVideoSource.h
@@ -2,6 +2,7 @@
 
 #pragma once
 
+#include "AtomicOnceFlag.h"
 #include "EpicRtcTrack.h"
 #include "EpicRtcVideoCapturer.h"
 #include "VideoSource.h"
@@ -30,11 +31,13 @@ namespace UE::PixelStreaming2
 		virtual void ForceKeyFrame() override;
 
 		TRefCountPtr<EpicRtcVideoTrackInterface> GetVideoTrack();
-
+		
+		inline bool TriggerDisconnectFlag() { return bDisconnectFlag.Trigger(); };
 	private:
 		FEpicRtcVideoSource(TRefCountPtr<EpicRtcVideoTrackInterface> InVideoTrack, TSharedPtr<FEpicRtcVideoCapturer> InVideoCapturer);
 
 	private:
 		TSharedPtr<FEpicRtcVideoCapturer> VideoCapturer;
+		FAtomicOnceFlag					  bDisconnectFlag;
 	};
 } // namespace UE::PixelStreaming2
diff --git a/Engine/Plugins/Media/PixelStreaming2/Source/PixelStreaming2RTC/Private/EpicRtcStreamer.cpp b/Engine/Plugins/Media/PixelStreaming2/Source/PixelStreaming2RTC/Private/EpicRtcStreamer.cpp
--- a/Engine/Plugins/Media/PixelStreaming2/Source/PixelStreaming2RTC/Private/EpicRtcStreamer.cpp
+++ b/Engine/Plugins/Media/PixelStreaming2/Source/PixelStreaming2RTC/Private/EpicRtcStreamer.cpp
@@ -649,16 +649,35 @@ namespace UE::PixelStreaming2
 
 		UE_LOGFMT(LogPixelStreaming2RTC, Verbose, "FEpicRtcStreamer::DeletePlayerSession(Participant [{0}])", PlayerId);
 
-		OnDataTrackClosed(PlayerId);
-		// FIXME (RTCP-7928): EpicRtc currently isn't broadcasting a stopped track state for remote track
-		OnAudioTrackClosed(PlayerId, true);
-		OnVideoTrackClosed(PlayerId, true);
-
 		if (ParticipantToDelete->VideoSource)
 		{
 			VideoSourceGroup->RemoveVideoSource(ParticipantToDelete->VideoSource.Get());
 		}
 
+		// FIXME (RTCP-7928): EpicRtc currently isn't broadcasting a stopped track state for remote track
+		if (ParticipantToDelete->VideoSink)
+		{
+			CloseVideoTrack(ParticipantToDelete, PlayerId, true);
+		}
+		if (ParticipantToDelete->VideoSource)
+		{
+			CloseVideoTrack(ParticipantToDelete, PlayerId, false);
+		}
+
+		if (ParticipantToDelete->AudioSink)
+		{
+			CloseAudioTrack(ParticipantToDelete, PlayerId, true);
+		}
+		if (ParticipantToDelete->AudioSource)
+		{
+			CloseAudioTrack(ParticipantToDelete, PlayerId, false);
+		}
+
+		if (ParticipantToDelete->DataTrack)
+		{
+			CloseDataTrack(ParticipantToDelete, PlayerId);
+		}
+
 		DoOnGameThread([StreamerId = this->StreamerId, PlayerId, bIsEmpty = Participants->IsEmpty()]() {
 			if (UPixelStreaming2Delegates* Delegates = UPixelStreaming2Delegates::Get())
 			{
@@ -733,7 +752,14 @@ namespace UE::PixelStreaming2
 	{
 		UE_LOGFMT(LogPixelStreaming2RTC, Verbose, "FEpicRtcStreamer::OnDataTrackClosed(Participant [{0}])", PlayerId);
 
-		if (TSharedPtr<FPlayerContext> Participant = Participants->FindRef(PlayerId); Participant.IsValid())
+		TSharedPtr<FPlayerContext> Participant = Participants->FindRef(PlayerId);
+
+		CloseDataTrack(Participant, PlayerId);
+	}
+
+	void FEpicRtcStreamer::CloseDataTrack(const TSharedPtr<FPlayerContext>& Participant, FString PlayerId)
+	{
+		if (Participant.IsValid())
 		{
 			Participant->DataTrack = nullptr;
 		}
@@ -779,21 +805,33 @@ namespace UE::PixelStreaming2
 
 	void FEpicRtcStreamer::OnAudioTrackClosed(FString PlayerId, bool bIsRemote)
 	{
-		UE_LOGFMT(LogPixelStreaming2RTC, Verbose, "FEpicRtcStreamer::OnAudioTrackClosed(Participant [{0}], IsRemote [{1}])", PlayerId, bIsRemote);
-
 		if (TSharedPtr<FPlayerContext> Participant = Participants->FindRef(PlayerId); Participant.IsValid())
 		{
-			if (bIsRemote)
+			CloseAudioTrack(Participant, PlayerId, bIsRemote);
+		}
+	}
+
+	void FEpicRtcStreamer::CloseAudioTrack(const TSharedPtr<FPlayerContext>& Participant, FString PlayerId, bool bIsRemote)
+	{
+		bool bCallDelegate = false;
+		if (bIsRemote)
+		{
+			if (TSharedPtr<FEpicRtcAudioSink> AudioSink = Participant->AudioSink)
 			{
+				bCallDelegate = !AudioSink->TriggerDisconnectFlag();
 				Participant->AudioSink = nullptr;
 			}
-			else
+		}
+		else
+		{
+			if (TSharedPtr<FEpicRtcAudioSource> AudioSource = Participant->AudioSource)
 			{
+				bCallDelegate = !AudioSource->TriggerDisconnectFlag();
 				Participant->AudioSource = nullptr;
 			}
 		}
 
-		if (UPixelStreaming2Delegates* Delegates = UPixelStreaming2Delegates::Get())
+		if (UPixelStreaming2Delegates* Delegates = UPixelStreaming2Delegates::Get(); bCallDelegate && Delegates)
 		{
 			// NOTE: Native delegates aren't required to be broadcast on game thread
 			Delegates->OnAudioTrackClosedNative.Broadcast(StreamerId, PlayerId, bIsRemote);
@@ -813,21 +851,34 @@ namespace UE::PixelStreaming2
 
 	void FEpicRtcStreamer::OnVideoTrackClosed(FString PlayerId, bool bIsRemote)
 	{
-		UE_LOGFMT(LogPixelStreaming2RTC, Verbose, "FEpicRtcStreamer::OnVideoTrackClosed(Participant [{0}], IsRemote [{1}])", PlayerId, bIsRemote);
-
 		if (TSharedPtr<FPlayerContext> Participant = Participants->FindRef(PlayerId); Participant.IsValid())
 		{
-			if (bIsRemote)
+			CloseVideoTrack(Participant, PlayerId, bIsRemote);
+		}
+	}
+
+	void FEpicRtcStreamer::CloseVideoTrack(const TSharedPtr<FPlayerContext>& Participant, FString PlayerId, bool bIsRemote)
+	{
+		bool bCallDelegate = false;
+		if (bIsRemote)
+		{
+			UE_LOGFMT(LogPixelStreaming2RTC, Verbose, "FEpicRtcStreamer::OnVideoTrackClosed Clear VideoSink (Participant [{0}], IsRemote [{1}])", PlayerId, bIsRemote);
+			if (TSharedPtr<FEpicRtcVideoSink> VideoSink = Participant->VideoSink)
 			{
+				bCallDelegate = !VideoSink->TriggerDisconnectFlag();
 				Participant->VideoSink = nullptr;
 			}
-			else
+		}
+		else
+		{
+			if (TSharedPtr<FEpicRtcVideoSource> VideoSource = Participant->VideoSource)
 			{
+				bCallDelegate = !VideoSource->TriggerDisconnectFlag();
 				Participant->VideoSource = nullptr;
 			}
 		}
 
-		if (UPixelStreaming2Delegates* Delegates = UPixelStreaming2Delegates::Get())
+		if (UPixelStreaming2Delegates* Delegates = UPixelStreaming2Delegates::Get(); bCallDelegate && Delegates)
 		{
 			// NOTE: Native delegates aren't required to be broadcast on game thread
 			Delegates->OnVideoTrackClosedNative.Broadcast(StreamerId, PlayerId, bIsRemote);
@@ -1532,7 +1583,10 @@ namespace UE::PixelStreaming2
 		{
 			if (bIsRemote)
 			{
-				Participant->VideoSink = FEpicRtcVideoSink::Create(VideoTrack);
+				// We received a remote track. We should now create a sink to handle receiving the frames.
+				// NOTE: We pass in nullptr as the track because if we store the track on the sink, EpicRtc will be unable to destroy it
+				// and webrtc will try to flush remaining frames during session removal. 
+				Participant->VideoSink = FEpicRtcVideoSink::Create(nullptr);
 			}
 			else
 			{
diff --git a/Engine/Plugins/Media/PixelStreaming2/Source/PixelStreaming2RTC/Private/EpicRtcStreamer.h b/Engine/Plugins/Media/PixelStreaming2/Source/PixelStreaming2RTC/Private/EpicRtcStreamer.h
--- a/Engine/Plugins/Media/PixelStreaming2/Source/PixelStreaming2RTC/Private/EpicRtcStreamer.h
+++ b/Engine/Plugins/Media/PixelStreaming2/Source/PixelStreaming2RTC/Private/EpicRtcStreamer.h
@@ -114,6 +114,9 @@ namespace UE::PixelStreaming2
 		void OnAudioTrackClosed(FString PlayerId, bool bIsRemote);
 		void OnVideoTrackOpen(FString PlayerId, bool bIsRemote);
 		void OnVideoTrackClosed(FString PlayerId, bool bIsRemote);
+		void CloseAudioTrack(const TSharedPtr<FPlayerContext>& Participant, FString PlayerId, bool bIsRemote);
+		void CloseDataTrack(const TSharedPtr<FPlayerContext>& Participant, FString PlayerId);
+		void CloseVideoTrack(const TSharedPtr<FPlayerContext>& Participant, FString PlayerId, bool bIsRemote);
 
 		void		   OnStatsReady(const FString& PlayerId, const EpicRtcConnectionStats& ConnectionStats);
 		void		   OnFrameCapturerCreated();
diff --git a/Engine/Plugins/Media/PixelStreaming2/Source/PixelStreaming2RTC/Private/Tests/DelegateTest.cpp b/Engine/Plugins/Media/PixelStreaming2/Source/PixelStreaming2RTC/Private/Tests/DelegateTest.cpp
--- a/Engine/Plugins/Media/PixelStreaming2/Source/PixelStreaming2RTC/Private/Tests/DelegateTest.cpp
+++ b/Engine/Plugins/Media/PixelStreaming2/Source/PixelStreaming2RTC/Private/Tests/DelegateTest.cpp
@@ -78,12 +78,10 @@ bool UPixelStreaming2DynamicDelegateTest::Init(UE::PixelStreaming2::DelegateTest
 		BIND_DELEGATE(Delegates->OnConnectedToSignallingServer, &UPixelStreaming2DynamicDelegateTest::OnConnectedToSignallingServer, TEXT("OnConnectedToSignallingServer"), FString)->Times(UE::PixelStreaming2::Exactly(1)).With(StreamerName);
 		BIND_DELEGATE(Delegates->OnDisconnectedFromSignallingServer, &UPixelStreaming2DynamicDelegateTest::OnDisconnectedFromSignallingServer, TEXT("OnDisconnectedFromSignallingServer"), FString)->Times(UE::PixelStreaming2::Exactly(1)).With(StreamerName);
 		BIND_DELEGATE(Delegates->OnNewConnection, &UPixelStreaming2DynamicDelegateTest::OnNewConnection, TEXT("OnNewConnection"), FString, TOptional<FString>)->Times(UE::PixelStreaming2::Exactly(Config.NumPlayers)).With(StreamerName, UE::PixelStreaming2::Any<FString>());
-		// TODO (Eden.Harris) This currently only fires once but should fire NumPlayers times.
-		BIND_DELEGATE(Delegates->OnClosedConnection, &UPixelStreaming2DynamicDelegateTest::OnClosedConnection, TEXT("OnClosedConnection"), FString, TOptional<FString>)->Times(UE::PixelStreaming2::Exactly(1)).With(StreamerName, UE::PixelStreaming2::Any<FString>());
+		BIND_DELEGATE(Delegates->OnClosedConnection, &UPixelStreaming2DynamicDelegateTest::OnClosedConnection, TEXT("OnClosedConnection"), FString, TOptional<FString>)->Times(UE::PixelStreaming2::Exactly(Config.NumPlayers)).With(StreamerName, UE::PixelStreaming2::Any<FString>());
 		BIND_DELEGATE(Delegates->OnAllConnectionsClosed, &UPixelStreaming2DynamicDelegateTest::OnAllConnectionsClosed, TEXT("OnAllConnectionsClosed"), FString)->Times(UE::PixelStreaming2::Exactly(1)).With(StreamerName);
 		BIND_DELEGATE(Delegates->OnDataTrackOpen, &UPixelStreaming2DynamicDelegateTest::OnDataTrackOpen, TEXT("OnDataTrackOpen"), FString, TOptional<FString>)->Times(UE::PixelStreaming2::Exactly(Config.NumPlayers)).With(StreamerName, UE::PixelStreaming2::Any<FString>());
-		// TODO (Eden.Harris) DataTrack closed does not currently fire
-		BIND_DELEGATE(Delegates->OnDataTrackClosed, &UPixelStreaming2DynamicDelegateTest::OnDataTrackClosed, TEXT("OnDataTrackClosed"), FString, TOptional<FString>)->Times(UE::PixelStreaming2::AtLeast(0)).With(StreamerName, UE::PixelStreaming2::Any<FString>());
+		BIND_DELEGATE(Delegates->OnDataTrackClosed, &UPixelStreaming2DynamicDelegateTest::OnDataTrackClosed, TEXT("OnDataTrackClosed"), FString, TOptional<FString>)->Times(UE::PixelStreaming2::Exactly(Config.NumPlayers)).With(StreamerName, UE::PixelStreaming2::Any<FString>());
 		BIND_DELEGATE(Delegates->OnStatChanged, &UPixelStreaming2DynamicDelegateTest::OnStatChanged, TEXT("OnStatChanged"), TOptional<FString>, TOptional<FName>, TOptional<float>)->Times(UE::PixelStreaming2::AtLeast(1)).With(UE::PixelStreaming2::Any<FString>(), UE::PixelStreaming2::Any<FName>(), UE::PixelStreaming2::Any<float>());
 		BIND_DELEGATE(Delegates->OnFallbackToSoftwareEncoding, &UPixelStreaming2DynamicDelegateTest::OnFallbackToSoftwareEncoding, TEXT("OnFallbackToSoftwareEncoding"))->Times(UE::PixelStreaming2::Exactly(Config.SoftwareEncodingCount));
 	}
@@ -220,26 +218,22 @@ namespace UE::PixelStreaming2
 			CreateSingleDelegateTest<decltype(Delegates->OnConnectedToSignallingServerNative), FString>(DelegatesMap, Delegates->OnConnectedToSignallingServerNative, TEXT("OnConnectedToSignallingServerNative"))->Times(Exactly(1)).With(StreamerName);
 			CreateSingleDelegateTest<decltype(Delegates->OnDisconnectedFromSignallingServerNative), FString>(DelegatesMap, Delegates->OnDisconnectedFromSignallingServerNative, TEXT("OnDisconnectedFromSignallingServerNative"))->Times(Exactly(1)).With(StreamerName);
 			CreateSingleDelegateTest<decltype(Delegates->OnNewConnectionNative), FString, TOptional<FString>>(DelegatesMap, Delegates->OnNewConnectionNative, TEXT("OnNewConnectionNative"))->Times(Exactly(Config.NumPlayers)).With(StreamerName, Any<FString>());
-			// TODO (Eden.Harris) This currently only fires once but should fire NumPlayers times.
-			CreateSingleDelegateTest<decltype(Delegates->OnClosedConnectionNative), FString, TOptional<FString>>(DelegatesMap, Delegates->OnClosedConnectionNative, TEXT("OnClosedConnectionNative"))->Times(Exactly(1)).With(StreamerName, Any<FString>());
+			CreateSingleDelegateTest<decltype(Delegates->OnClosedConnectionNative), FString, TOptional<FString>>(DelegatesMap, Delegates->OnClosedConnectionNative, TEXT("OnClosedConnectionNative"))->Times(Exactly(Config.NumPlayers)).With(StreamerName, Any<FString>());
 			CreateSingleDelegateTest<decltype(Delegates->OnAllConnectionsClosedNative), FString>(DelegatesMap, Delegates->OnAllConnectionsClosedNative, TEXT("OnAllConnectionsClosedNative"))->Times(Exactly(1)).With(StreamerName);
 			CreateSingleDelegateTest<decltype(Delegates->OnDataTrackOpenNative), FString, TOptional<FString>>(DelegatesMap, Delegates->OnDataTrackOpenNative, TEXT("OnDataTrackOpenNative"))->Times(Exactly(Config.NumPlayers)).With(StreamerName, Any<FString>());
-			// TODO (Eden.Harris) DataTrack closed does not currently fire
-			CreateSingleDelegateTest<decltype(Delegates->OnDataTrackClosedNative), FString, TOptional<FString>>(DelegatesMap, Delegates->OnDataTrackClosedNative, TEXT("OnDataTrackClosedNative"))->Times(AtLeast(0)).With(StreamerName, Any<FString>());
+			CreateSingleDelegateTest<decltype(Delegates->OnDataTrackClosedNative), FString, TOptional<FString>>(DelegatesMap, Delegates->OnDataTrackClosedNative, TEXT("OnDataTrackClosedNative"))->Times(Exactly(Config.NumPlayers)).With(StreamerName, Any<FString>());
 
 			const int NumCalls = Config.bIsBidirectional ? Config.NumPlayers * 2: Config.NumPlayers;
 
 			auto& OnVideoTrackOpenNativeDelegate = CreateSingleDelegateTest<decltype(Delegates->OnVideoTrackOpenNative), FString, TOptional<FString>, bool>(DelegatesMap, Delegates->OnVideoTrackOpenNative, TEXT("OnVideoTrackOpenNative"))->Times(Exactly(NumCalls)).With(StreamerName, Any<FString>(), bIsRemote);
-			// TODO (Eden.Harris) This should be called twice, once for local and remote. 
-			// This is also not being fired on linux or Mac so is disabled for now.
-			// auto& OnVideoTrackClosedNativeDelegate = CreateSingleDelegateTest<decltype(Delegates->OnVideoTrackClosedNative), FString, TOptional<FString>, bool>(DelegatesMap, Delegates->OnVideoTrackClosedNative, TEXT("OnVideoTrackClosedNative"))->Times(Exactly(Config.NumPlayers)).With(StreamerName, Any<FString>(), bIsRemote);
+			auto& OnVideoTrackClosedNativeDelegate = CreateSingleDelegateTest<decltype(Delegates->OnVideoTrackClosedNative), FString, TOptional<FString>, bool>(DelegatesMap, Delegates->OnVideoTrackClosedNative, TEXT("OnVideoTrackClosedNative"))->Times(Exactly(NumCalls)).With(StreamerName, Any<FString>(), bIsRemote);
 			auto& OnAudioTrackOpenNativeDelegate = CreateSingleDelegateTest<decltype(Delegates->OnAudioTrackOpenNative), FString, TOptional<FString>, bool>(DelegatesMap, Delegates->OnAudioTrackOpenNative, TEXT("OnAudioTrackOpenNative"))->Times(Exactly(NumCalls)).With( StreamerName, Any<FString>(), bIsRemote);
 			auto& OnAudioTrackClosedNativeDelegate = CreateSingleDelegateTest<decltype(Delegates->OnAudioTrackClosedNative), FString, TOptional<FString>, bool>(DelegatesMap, Delegates->OnAudioTrackClosedNative, TEXT("OnAudioTrackClosedNative"))->Times(Exactly(NumCalls)).With(StreamerName, Any<FString>(), bIsRemote);
 			
 			if(Config.bIsBidirectional)
 			{
 				OnVideoTrackOpenNativeDelegate.With(StreamerName, Any<FString>(), !bIsRemote);
-				//OnVideoTrackClosedNativeDelegate.With(StreamerName, Any<FString>(), !bIsRemote);
+				OnVideoTrackClosedNativeDelegate.With(StreamerName, Any<FString>(), !bIsRemote);
 				OnAudioTrackOpenNativeDelegate.With(StreamerName, Any<FString>(), !bIsRemote);
 				OnAudioTrackClosedNativeDelegate.With(StreamerName, Any<FString>(), !bIsRemote);
 			}
@@ -353,7 +347,7 @@ namespace UE::PixelStreaming2
 		});
 
 		ADD_LATENT_AUTOMATION_COMMAND(FExecuteLambda([Streamer]() { Streamer->StartStreaming(); }))
-		ADD_LATENT_AUTOMATION_COMMAND(FWAitForBoolOrTimeout(TEXT("Check streaming started"), 5.0, Streamer, bStreamingStarted, true))
+		ADD_LATENT_AUTOMATION_COMMAND(FWaitForBoolOrTimeout(TEXT("Check streaming started"), 5.0, Streamer, bStreamingStarted, true))
 
 		for (TSharedPtr<FMockPlayer>& Player : Players)
 		{
@@ -369,7 +363,7 @@ namespace UE::PixelStreaming2
 		// Wait 1 second to ensure any websocket message have correctly flowed
 		ADD_LATENT_AUTOMATION_COMMAND(FWaitSeconds(1.0))
 		ADD_LATENT_AUTOMATION_COMMAND(FExecuteLambda([Streamer]() { Streamer->StopStreaming(); }))
-		ADD_LATENT_AUTOMATION_COMMAND(FWAitForBoolOrTimeout(TEXT("Check disconnected"), 5.0, Streamer, bStreamingDisconnected, true))
+		ADD_LATENT_AUTOMATION_COMMAND(FWaitForBoolOrTimeout(TEXT("Check disconnected"), 5.0, Streamer, bStreamingDisconnected, true))
 
 		ADD_LATENT_AUTOMATION_COMMAND(FCleanupAllPlayers(SignallingServer, Streamer, Players))
 		
@@ -383,8 +377,6 @@ namespace UE::PixelStreaming2
 		}
 	}
 
-	// TODO (Eden.Harris) RTCP-8326 This test is failing to fires some delegates.
-#if 0
 	IMPLEMENT_SIMPLE_AUTOMATION_TEST(FPS2DelegateSoftwareFallbackTest, "System.Plugins.PixelStreaming2.FPS2DelegateSoftwareFallbackTest", EAutomationTestFlags::EditorContext | EAutomationTestFlags::ClientContext | EAutomationTestFlags::ProductFilter)
 	bool FPS2DelegateSoftwareFallbackTest::RunTest(const FString&)
 	{
@@ -396,7 +388,7 @@ namespace UE::PixelStreaming2
 			const EVideoCodec PrevCodec = UE::PixelStreaming2::GetEnumFromCVar<EVideoCodec>(UPixelStreaming2PluginSettings::CVarEncoderCodec);
 			UPixelStreaming2PluginSettings::CVarEncoderMaxSessions.AsVariable()->Set(0);
 			UPixelStreaming2PluginSettings::CVarEncoderCodec.AsVariable()->Set(*UE::PixelStreaming2::GetCVarStringFromEnum(EVideoCodec::H264));
-			Result = RunDelegateTest<FDelegateNativeTest>(1, 1);
+			RunDelegateTest<FDelegateNativeTest>({ 1, 1, false });
 
 			// Reset to previous
 			ADD_LATENT_AUTOMATION_COMMAND(FExecuteLambda([PrevPixelStreamingEncoderMaxSessions, PrevCodec]() { 
@@ -409,9 +401,8 @@ namespace UE::PixelStreaming2
 			UE_LOGFMT(LogPixelStreaming2RTC, Log, "FPS2DelegateSoftwareFallbackTest requires Nvidia GPU to test Software fallback");
 		}
 		
-		return Result;
+		return true;
 	}
-#endif
 
 	IMPLEMENT_SIMPLE_AUTOMATION_TEST(FPS2DelegateNativeSingleTest, "System.Plugins.PixelStreaming2.FPS2DelegateNativeSingleTest", EAutomationTestFlags::EditorContext | EAutomationTestFlags::ClientContext | EAutomationTestFlags::ProductFilter)
 	bool FPS2DelegateNativeSingleTest::RunTest(const FString&)
diff --git a/Engine/Plugins/Media/PixelStreaming2/Source/PixelStreaming2RTC/Private/Tests/EpicRtcStreamerTest.cpp b/Engine/Plugins/Media/PixelStreaming2/Source/PixelStreaming2RTC/Private/Tests/EpicRtcStreamerTest.cpp
--- a/Engine/Plugins/Media/PixelStreaming2/Source/PixelStreaming2RTC/Private/Tests/EpicRtcStreamerTest.cpp
+++ b/Engine/Plugins/Media/PixelStreaming2/Source/PixelStreaming2RTC/Private/Tests/EpicRtcStreamerTest.cpp
@@ -60,14 +60,14 @@ namespace UE::PixelStreaming2
 
 		// Start streaming
 		ADD_LATENT_AUTOMATION_COMMAND(FExecuteLambda([Streamer]() { Streamer->StartStreaming(); }))
-		ADD_LATENT_AUTOMATION_COMMAND(FWaitAndCheckStreamerBool(TEXT("Check streaming started"), 5.0, Streamer, bStreamingStarted, true))
-		ADD_LATENT_AUTOMATION_COMMAND(FWaitAndCheckStreamerBool(TEXT("Check connected to signalling server"), 5.0, Streamer, bConnectedToSignallingServerNative, true))
+		ADD_LATENT_AUTOMATION_COMMAND(FWaitForBoolOrTimeout(TEXT("Check streaming started"), 5.0, Streamer, bStreamingStarted, true))
+		ADD_LATENT_AUTOMATION_COMMAND(FWaitForBoolOrTimeout(TEXT("Check connected to signalling server"), 5.0, Streamer, bConnectedToSignallingServerNative, true))
 		// Wait 1 second to ensure any websocket message have correctly flowed
 		ADD_LATENT_AUTOMATION_COMMAND(FWaitSeconds(1.0))
 		// Stop streaming
 		ADD_LATENT_AUTOMATION_COMMAND(FExecuteLambda([Streamer]() { Streamer->StopStreaming(); }))
-		ADD_LATENT_AUTOMATION_COMMAND(FWaitAndCheckStreamerBool(TEXT("Check streaming stopped"), 5.0, Streamer, bStreamingStopped, true))
-		ADD_LATENT_AUTOMATION_COMMAND(FWaitAndCheckStreamerBool(TEXT("Check disconnected from signalling server"), 5.0, Streamer, bDisconnectedFromSignallingServerNative, true))
+		ADD_LATENT_AUTOMATION_COMMAND(FWaitForBoolOrTimeout(TEXT("Check streaming stopped"), 5.0, Streamer, bStreamingStopped, true))
+		ADD_LATENT_AUTOMATION_COMMAND(FWaitForBoolOrTimeout(TEXT("Check disconnected from signalling server"), 5.0, Streamer, bDisconnectedFromSignallingServerNative, true))
 		// Cleanup
 		ADD_LATENT_AUTOMATION_COMMAND(FCleanupAll(SignallingServer, Streamer, nullptr))
 
@@ -120,14 +120,14 @@ namespace UE::PixelStreaming2
 
 		// Start streaming
 		ADD_LATENT_AUTOMATION_COMMAND(FExecuteLambda([Streamer]() { Streamer->StartStreaming(); }))
-		ADD_LATENT_AUTOMATION_COMMAND(FWaitAndCheckStreamerBool(TEXT("Check streaming started"), 5.0, Streamer, bStreamingStarted, true))
-		ADD_LATENT_AUTOMATION_COMMAND(FWaitAndCheckStreamerBool(TEXT("Check connected to signalling server"), 5.0, Streamer, bConnectedToSignallingServerNative, true))
+		ADD_LATENT_AUTOMATION_COMMAND(FWaitForBoolOrTimeout(TEXT("Check streaming started"), 5.0, Streamer, bStreamingStarted, true))
+		ADD_LATENT_AUTOMATION_COMMAND(FWaitForBoolOrTimeout(TEXT("Check connected to signalling server"), 5.0, Streamer, bConnectedToSignallingServerNative, true))
 		// Wait 1 second to ensure any websocket message have correctly flowed
 		ADD_LATENT_AUTOMATION_COMMAND(FWaitSeconds(1.0))
 		// Reset the signalling server, this will trigger the reconnection flow
 		ADD_LATENT_AUTOMATION_COMMAND(FExecuteLambda([SignallingServer]() { SignallingServer->Stop(); }))
-		ADD_LATENT_AUTOMATION_COMMAND(FWaitAndCheckStreamerBool(TEXT("Check streaming stopped from SS going away"), 5.0, Streamer, bStreamingStopped, true))
-		ADD_LATENT_AUTOMATION_COMMAND(FWaitAndCheckStreamerBool(TEXT("Check disconnected from signalling server from SS going away"), 5.0, Streamer, bDisconnectedFromSignallingServerNative, true))
+		ADD_LATENT_AUTOMATION_COMMAND(FWaitForBoolOrTimeout(TEXT("Check streaming stopped from SS going away"), 5.0, Streamer, bStreamingStopped, true))
+		ADD_LATENT_AUTOMATION_COMMAND(FWaitForBoolOrTimeout(TEXT("Check disconnected from signalling server from SS going away"), 5.0, Streamer, bDisconnectedFromSignallingServerNative, true))
 		// Reset the state variables
 		ADD_LATENT_AUTOMATION_COMMAND(FExecuteLambda([bStreamingStarted, bConnectedToSignallingServerNative, bStreamingStopped, bDisconnectedFromSignallingServerNative]() {
 			*(bStreamingStarted.Get()) = false;
@@ -147,14 +147,14 @@ namespace UE::PixelStreaming2
 		}))
 		// Wait to ensure any websocket message have correctly flowed. Time is 2 times the reconnect interval for safety
 		ADD_LATENT_AUTOMATION_COMMAND(FWaitSeconds(2 * UPixelStreaming2PluginSettings::CVarSignalingReconnectInterval.GetValueOnAnyThread()))
-		ADD_LATENT_AUTOMATION_COMMAND(FWaitAndCheckStreamerBool(TEXT("Check streaming restarted"), 5.0, Streamer, bStreamingStarted, true))
-		ADD_LATENT_AUTOMATION_COMMAND(FWaitAndCheckStreamerBool(TEXT("Check reconnected to signalling server"), 5.0, Streamer, bConnectedToSignallingServerNative, true))
+		ADD_LATENT_AUTOMATION_COMMAND(FWaitForBoolOrTimeout(TEXT("Check streaming restarted"), 5.0, Streamer, bStreamingStarted, true))
+		ADD_LATENT_AUTOMATION_COMMAND(FWaitForBoolOrTimeout(TEXT("Check reconnected to signalling server"), 5.0, Streamer, bConnectedToSignallingServerNative, true))
 		// Wait 1 second to ensure any websocket message have correctly flowed
 		ADD_LATENT_AUTOMATION_COMMAND(FWaitSeconds(1.0))
 		// Stop streaming
 		ADD_LATENT_AUTOMATION_COMMAND(FExecuteLambda([Streamer]() { Streamer->StopStreaming(); }))
-		ADD_LATENT_AUTOMATION_COMMAND(FWaitAndCheckStreamerBool(TEXT("Check streaming stopped"), 5.0, Streamer, bStreamingStopped, true))
-		ADD_LATENT_AUTOMATION_COMMAND(FWaitAndCheckStreamerBool(TEXT("Check disconnected from signalling server"), 5.0, Streamer, bDisconnectedFromSignallingServerNative, true))
+		ADD_LATENT_AUTOMATION_COMMAND(FWaitForBoolOrTimeout(TEXT("Check streaming stopped"), 5.0, Streamer, bStreamingStopped, true))
+		ADD_LATENT_AUTOMATION_COMMAND(FWaitForBoolOrTimeout(TEXT("Check disconnected from signalling server"), 5.0, Streamer, bDisconnectedFromSignallingServerNative, true))
 
 		// Cleanup
 		ADD_LATENT_AUTOMATION_COMMAND(FCleanupAll(SignallingServer, Streamer, nullptr))
diff --git a/Engine/Plugins/Media/PixelStreaming2/Source/PixelStreaming2RTC/Private/Tests/MediaDirectionTest.cpp b/Engine/Plugins/Media/PixelStreaming2/Source/PixelStreaming2RTC/Private/Tests/MediaDirectionTest.cpp
--- a/Engine/Plugins/Media/PixelStreaming2/Source/PixelStreaming2RTC/Private/Tests/MediaDirectionTest.cpp
+++ b/Engine/Plugins/Media/PixelStreaming2/Source/PixelStreaming2RTC/Private/Tests/MediaDirectionTest.cpp
@@ -60,7 +60,7 @@ namespace UE::PixelStreaming2
 		TSharedPtr<bool> bRemoteTrack = MakeShared<bool>(false);
 
 		ADD_LATENT_AUTOMATION_COMMAND(FExecuteLambda([Streamer]() { Streamer->StartStreaming(); }))
-		ADD_LATENT_AUTOMATION_COMMAND(FWaitAndCheckStreamerBool(TEXT("Check streaming started"), 5.0, Streamer, bStreamingStarted, true))
+		ADD_LATENT_AUTOMATION_COMMAND(FWaitForBoolOrTimeout(TEXT("Check streaming started"), 5.0, Streamer, bStreamingStarted, true))
 		ADD_LATENT_AUTOMATION_COMMAND(FExecuteLambda([Player, PlayerPort]() { Player->Connect(PlayerPort); }))
 		ADD_LATENT_AUTOMATION_COMMAND(FSubscribePlayerAfterStreamerConnectedOrTimeout(5.0, Streamer, Player, StreamerName))
 		// Wait 5 seconds to ensure connection is fully established
diff --git a/Engine/Plugins/Media/PixelStreaming2/Source/PixelStreaming2RTC/Private/Tests/TestUtils.cpp b/Engine/Plugins/Media/PixelStreaming2/Source/PixelStreaming2RTC/Private/Tests/TestUtils.cpp
--- a/Engine/Plugins/Media/PixelStreaming2/Source/PixelStreaming2RTC/Private/Tests/TestUtils.cpp
+++ b/Engine/Plugins/Media/PixelStreaming2/Source/PixelStreaming2RTC/Private/Tests/TestUtils.cpp
@@ -145,10 +145,17 @@ namespace UE::PixelStreaming2
 				UE_LOG(LogPixelStreaming2RTC, Log, TEXT("%s"), *WidthTestString);
 			}
 
-			EpicRtcPixelFormat Format = Buffer->GetFormat();
-			if (Format != EpicRtcPixelFormat::I420)
+			FEpicRtcVideoBuffer* const BaseFrameBuffer = static_cast<FEpicRtcVideoBuffer*>(Buffer.GetReference());
+			if (!BaseFrameBuffer)
+			{
+				UE_LOG(LogPixelStreaming2RTC, Error, TEXT("Invalid frame buffer"));
+				return true;
+			}
+
+			if (BaseFrameBuffer->GetBufferFormat() != PixelCaptureBufferFormat::FORMAT_I420)
 			{
 				UE_LOG(LogPixelStreaming2RTC, Error, TEXT("Invalid Pixel Format"));
+				return true;
 			}
 
 			uint8_t* DataY = reinterpret_cast<uint8_t*>(Buffer->GetData());
@@ -248,7 +255,7 @@ namespace UE::PixelStreaming2
 		return false; // We keep updating this until the timeout is complete and then we check that there's no tracks added
 	}
 
-	bool FWAitForBoolOrTimeout::Update()
+	bool FWaitForBoolOrTimeout::Update()
 	{
 		double DeltaTime = FPlatformTime::Seconds() - StartTime;
 		if (DeltaTime > WaitSeconds)
@@ -335,7 +342,7 @@ namespace UE::PixelStreaming2
 
 	bool FCleanupAllPlayers::Update()
 	{
-		for(TSharedPtr<FMockPlayer>& Player : OutPlayers)
+		for (TSharedPtr<FMockPlayer>& Player : OutPlayers)
 		{
 			if (Player)
 			{
diff --git a/Engine/Plugins/Media/PixelStreaming2/Source/PixelStreaming2RTC/Private/Tests/TestUtils.h b/Engine/Plugins/Media/PixelStreaming2/Source/PixelStreaming2RTC/Private/Tests/TestUtils.h
--- a/Engine/Plugins/Media/PixelStreaming2/Source/PixelStreaming2RTC/Private/Tests/TestUtils.h
+++ b/Engine/Plugins/Media/PixelStreaming2/Source/PixelStreaming2RTC/Private/Tests/TestUtils.h
@@ -88,7 +88,7 @@ namespace UE::PixelStreaming2
 	DEFINE_LATENT_AUTOMATION_COMMAND_THREE_PARAMETER(FWaitForPlayerTrackOrTimeout, double, TimeoutSeconds, TSharedPtr<FMockPlayer>, Player, TSharedPtr<bool>, bComplete);
 	DEFINE_LATENT_AUTOMATION_COMMAND_FOUR_PARAMETER(FWaitAndCheckBool, double, WaitSeconds, TSharedPtr<FMockPlayer>, Player, TSharedPtr<bool>, bCheck, bool, bExpectedValue);
 	DEFINE_LATENT_AUTOMATION_COMMAND_FIVE_PARAMETER(FWaitAndCheckStreamerBool, const FString, CheckName, double, WaitSeconds, TSharedPtr<IPixelStreaming2Streamer>, Streamer, TSharedPtr<bool>, bCheck, bool, bExpectedValue);
-	DEFINE_LATENT_AUTOMATION_COMMAND_FIVE_PARAMETER(FWAitForBoolOrTimeout, const FString, CheckName, double, WaitSeconds, TSharedPtr<IPixelStreaming2Streamer>, Streamer, TSharedPtr<bool>, bCheck, bool, bExpectedValue);
+	DEFINE_LATENT_AUTOMATION_COMMAND_FIVE_PARAMETER(FWaitForBoolOrTimeout, const FString, CheckName, double, WaitSeconds, TSharedPtr<IPixelStreaming2Streamer>, Streamer, TSharedPtr<bool>, bCheck, bool, bExpectedValue);
 	DEFINE_LATENT_AUTOMATION_COMMAND_ONE_PARAMETER(FExecuteLambda, const TFunction<void()>, Func);
 
 	TSharedPtr<IPixelStreaming2Streamer>			CreateStreamer(const FString& StreamerName, int StreamerPort);
