[PS2] Add: MediaPlate based stream playback. Replaces the UPixelStreaming2Peer component

diff --git a/Engine/Plugins/Media/PixelStreaming2/Source/PixelStreaming2RTC/Private/PixelStreaming2RTCModule.cpp b/Engine/Plugins/Media/PixelStreaming2/Source/PixelStreaming2RTC/Private/PixelStreaming2RTCModule.cpp
--- a/Engine/Plugins/Media/PixelStreaming2/Source/PixelStreaming2RTC/Private/PixelStreaming2RTCModule.cpp
+++ b/Engine/Plugins/Media/PixelStreaming2/Source/PixelStreaming2RTC/Private/PixelStreaming2RTCModule.cpp
@@ -2,63 +2,7 @@
 
 #include "PixelStreaming2RTCModule.h"
 
-#include "UtilsCoder.h"
 #include "CoreMinimal.h"
-#include "UtilsCore.h"
-#include "Engine/GameViewportClient.h"
-#include "Engine/Texture2D.h"
-#include "IPixelStreaming2Module.h"
-#include "IPixelStreaming2InputModule.h"
-#include "Logging.h"
-#include "Modules/ModuleManager.h"
-#include "PixelStreaming2Delegates.h"
-#include "PixelStreaming2Utils.h"
-#include "PixelStreaming2PluginSettings.h"
-#include "Slate/SceneViewport.h"
-#include "EpicRtcStreamer.h"
-#include "UObject/UObjectIterator.h"
-#include "UtilsCommon.h"
-
-#if PLATFORM_LINUX
-	#include "CudaModule.h"
-#endif
-
-#if PLATFORM_WINDOWS
-	#include "Windows/AllowWindowsPlatformTypes.h"
-THIRD_PARTY_INCLUDES_START
-	#include <VersionHelpers.h>
-THIRD_PARTY_INCLUDES_END
-	#include "Windows/HideWindowsPlatformTypes.h"
-#endif
-
-#include "AudioDeviceManager.h"
-#include "RenderingThread.h"
-#include "Misc/CommandLine.h"
-#include "Misc/Parse.h"
-#include "RendererInterface.h"
-#include "Rendering/SlateRenderer.h"
-#include "Framework/Application/SlateApplication.h"
-#include "Misc/ConfigCacheIni.h"
-#include "GameFramework/GameModeBase.h"
-#include "Dom/JsonObject.h"
-#include "Misc/App.h"
-#include "Misc/MessageDialog.h"
-#include "Async/Async.h"
-#include "Engine/Engine.h"
-#include "WebSocketsModule.h"
-
-#include "Video/Encoders/Configs/VideoEncoderConfigH264.h"
-#include "Video/Encoders/Configs/VideoEncoderConfigAV1.h"
-
-#if !UE_BUILD_SHIPPING
-	#include "DrawDebugHelpers.h"
-#endif
-
-#include "VideoProducerBackBuffer.h"
-#include "VideoProducerMediaCapture.h"
-#include "Engine/GameEngine.h"
-#include "Stats.h"
-#include "UtilsString.h"
 
 #include "EpicRtcAllocator.h"
 #include "EpicRtcAudioCapturer.h"
@@ -66,6 +10,23 @@ THIRD_PARTY_INCLUDES_END
 #include "EpicRtcVideoEncoderInitializer.h"
 #include "EpicRtcVideoDecoderInitializer.h"
 #include "EpicRtcWebsocketFactory.h"
+#include "IMediaModule.h"
+#include "IPixelStreaming2Module.h"
+#include "Logging.h"
+#include "Misc/CommandLine.h"
+#include "Misc/Parse.h"
+#include "PixelStreaming2Delegates.h"
+#include "PixelStreaming2PluginSettings.h"
+#include "PixelStreaming2RTCPlayer.h"
+#include "PixelStreaming2Utils.h"
+#include "RendererInterface.h"
+#include "Stats.h"
+#include "UtilsCommon.h"
+#include "UtilsCoder.h"
+#include "UtilsCore.h"
+#include "Video/Encoders/Configs/VideoEncoderConfigH264.h"
+#include "Video/Encoders/Configs/VideoEncoderConfigAV1.h"
+#include "WebSocketsModule.h"
 
 namespace UE::PixelStreaming2
 {
@@ -185,6 +146,11 @@ namespace UE::PixelStreaming2
 			});
 		}
 
+		if (IMediaModule* MediaModulePtr = FModuleManager::LoadModulePtr<IMediaModule>("Media"); MediaModulePtr != nullptr)
+		{
+			MediaModulePtr->RegisterPlayerFactory(*this);
+		}
+
 		bStartupCompleted = true;
 	}
 
@@ -200,6 +166,11 @@ namespace UE::PixelStreaming2
 			return;
 		}
 
+		if (IMediaModule* MediaModulePtr = FModuleManager::LoadModulePtr<IMediaModule>("Media"); MediaModulePtr != nullptr)
+		{
+			MediaModulePtr->UnregisterPlayerFactory(*this);
+		}
+
 		AudioMixingCapturer.Reset();
 
 		TickConferenceTask.Reset();
@@ -247,6 +218,49 @@ namespace UE::PixelStreaming2
 	 * End IPixelStreaming2RTCModule implementation
 	 */
 
+	/**
+	 * IMediaPlayerFactory implementation
+	 */
+	TSharedPtr<IMediaPlayer> FPixelStreaming2RTCModule::CreatePlayer(IMediaEventSink& EventSink)
+	{
+		return MakeShared<FPixelStreaming2RTCStreamPlayer>();
+	}
+
+	FName FPixelStreaming2RTCModule::GetPlayerName() const
+	{
+		return "PixelStreaming2RTC";
+	}
+
+	FText FPixelStreaming2RTCModule::GetDisplayName() const
+	{
+		return NSLOCTEXT("PixelStreaming2", "MediaPlayerFactory", "PixelStreaming2 RTC Stream Player");
+	}
+
+	bool FPixelStreaming2RTCModule::CanPlayUrl(const FString& Url, const IMediaOptions* Options, TArray<FText>* OutWarnings, TArray<FText>* OutErrors) const
+	{
+		return Url.StartsWith(TEXT("ws://")) || Url.StartsWith(TEXT("wss://"));
+	}
+
+	bool FPixelStreaming2RTCModule::SupportsFeature(EMediaFeature Feature) const
+	{
+		return Feature == EMediaFeature::VideoSamples || Feature == EMediaFeature::AudioSamples;
+	}
+
+	FGuid FPixelStreaming2RTCModule::GetPlayerPluginGUID() const
+	{
+		static FGuid PlayerPluginGUID = FGuid::NewGuid();
+		return PlayerPluginGUID;
+	}
+
+	const TArray<FString>& FPixelStreaming2RTCModule::GetSupportedPlatforms() const
+	{
+		static TArray<FString> SupportedPlatforms = { TEXT("Windows"), TEXT("Linux"), TEXT("Mac") };
+		return SupportedPlatforms;
+	}
+	/**
+	 * End IMediaPlayerFactory implementation
+	 */
+
 	TSharedPtr<FEpicRtcAudioCapturer> FPixelStreaming2RTCModule::GetAudioCapturer()
 	{
 		if (!AudioMixingCapturer)
diff --git a/Engine/Plugins/Media/PixelStreaming2/Source/PixelStreaming2RTC/Private/PixelStreaming2RTCModule.h b/Engine/Plugins/Media/PixelStreaming2/Source/PixelStreaming2RTC/Private/PixelStreaming2RTCModule.h
--- a/Engine/Plugins/Media/PixelStreaming2/Source/PixelStreaming2RTC/Private/PixelStreaming2RTCModule.h
+++ b/Engine/Plugins/Media/PixelStreaming2/Source/PixelStreaming2RTC/Private/PixelStreaming2RTCModule.h
@@ -2,27 +2,23 @@
 
 #pragma once
 
-#include "IPixelStreaming2RTCModule.h"
 #include "EpicRtcConferenceUtils.h"
 #include "EpicRtcStatsCollector.h"
 #include "EpicRtcStreamer.h"
+#include "IMediaPlayerFactory.h"
+#include "IPixelStreaming2RTCModule.h"
 
 #include "epic_rtc/core/platform.h"
 #include "epic_rtc/plugins/signalling/signalling_type.h"
 
-class UPixelStreaming2Input;
-class SWindow;
-
 namespace UE::PixelStreaming2
 {
-	class FVideoInputBackBuffer;
-	class FVideoSourceGroup;
 	class FEpicRtcWebsocketFactory;
 
 	/**
 	 * This plugin allows the back buffer to be sent as a compressed video across a network.
 	 */
-	class FPixelStreaming2RTCModule : public IPixelStreaming2RTCModule
+	class FPixelStreaming2RTCModule : public IPixelStreaming2RTCModule, public IMediaPlayerFactory
 	{
 	public:
 		static FPixelStreaming2RTCModule* GetModule();
@@ -34,9 +30,19 @@ namespace UE::PixelStreaming2
 		virtual bool		 IsReady() override;
 		// End IPixelStreaming2RTCModule
 
-		TSharedPtr<class FEpicRtcAudioCapturer>	 GetAudioCapturer();
+		// Begin IMediaPlayerFactory
+		virtual TSharedPtr<IMediaPlayer> CreatePlayer(IMediaEventSink& EventSink) override;
+		virtual FName					 GetPlayerName() const override;
+		virtual FText					 GetDisplayName() const override;
+		virtual bool					 CanPlayUrl(const FString& Url, const IMediaOptions* Options, TArray<FText>* OutWarnings, TArray<FText>* OutErrors) const override;
+		virtual bool					 SupportsFeature(EMediaFeature Feature) const override;
+		virtual FGuid					 GetPlayerPluginGUID() const override;
+		virtual const TArray<FString>&	 GetSupportedPlatforms() const override;
+		// End IMediaPlayerFactory
+
+		TSharedPtr<class FEpicRtcAudioCapturer>	  GetAudioCapturer();
 		TRefCountPtr<EpicRtcConferenceInterface>& GetEpicRtcConference() { return EpicRtcConference; }
-		TRefCountPtr<FEpicRtcStatsCollector>&	 GetStatsCollector() { return StatsCollector; }
+		TRefCountPtr<FEpicRtcStatsCollector>&	  GetStatsCollector() { return StatsCollector; }
 
 	private:
 		// Begin IModuleInterface
diff --git a/Engine/Plugins/Media/PixelStreaming2/Source/PixelStreaming2RTC/Private/PixelStreaming2RTCPlayer.cpp b/Engine/Plugins/Media/PixelStreaming2/Source/PixelStreaming2RTC/Private/PixelStreaming2RTCPlayer.cpp
new file mode 100644
--- /dev/null
+++ b/Engine/Plugins/Media/PixelStreaming2/Source/PixelStreaming2RTC/Private/PixelStreaming2RTCPlayer.cpp
@@ -0,0 +1,888 @@
+// Copyright Epic Games, Inc. All Rights Reserved.
+
+#include "PixelStreaming2RTCPlayer.h"
+
+#include "GenericPlatform/GenericPlatformHttp.h"
+#include "IMediaAudioSample.h"
+#include "IMediaTextureSample.h"
+#include "MediaObjectPool.h"
+#include "Logging.h"
+#include "UtilsString.h"
+#include "PixelStreaming2RTCModule.h"
+#include "Templates/PointerVariants.h"
+
+namespace UE::PixelStreaming2
+{
+	class FPixelStreaming2TextureSample : public IMediaTextureSample, public IMediaPoolable
+	{
+	public:
+		FPixelStreaming2TextureSample()
+			: Time(FTimespan::Zero())
+		{
+		}
+
+		void Initialize(TRefCountPtr<FRHITexture> InTexture, FIntPoint InDisplaySize, FIntPoint InTotalSize, FTimespan InTime, FTimespan InDuration)
+		{
+			Texture = InTexture;
+			Time = InTime;
+			DisplaySize = InDisplaySize;
+			TotalSize = InTotalSize;
+			Duration = InDuration;
+		}
+
+		//~ IMediaTextureSample interface
+		virtual const void* GetBuffer() override
+		{
+			return nullptr;
+		}
+
+		virtual FIntPoint GetDim() const override
+		{
+			return TotalSize;
+		}
+
+		virtual FTimespan GetDuration() const override
+		{
+			return Duration;
+		}
+
+		virtual EMediaTextureSampleFormat GetFormat() const override
+		{
+			return EMediaTextureSampleFormat::CharBGRA;
+		}
+
+		virtual FIntPoint GetOutputDim() const override
+		{
+			return DisplaySize;
+		}
+
+		virtual uint32 GetStride() const override
+		{
+			if (!Texture.IsValid())
+			{
+				return 0;
+			}
+			return Texture->GetSizeX();
+		}
+
+		virtual FRHITexture* GetTexture() const override
+		{
+			return Texture.GetReference();
+		}
+
+		virtual FMediaTimeStamp GetTime() const override
+		{
+			return FMediaTimeStamp(Time);
+		}
+
+		virtual bool IsCacheable() const override
+		{
+			return true;
+		}
+
+		virtual bool IsOutputSrgb() const override
+		{
+			return true;
+		}
+
+	public:
+		//~ IMediaPoolable interface
+		virtual void ShutdownPoolable() override
+		{
+			// Drop reference to the texture. It should be released by the outside system.
+			Texture = nullptr;
+			Time = FTimespan::Zero();
+		}
+
+	private:
+		TRefCountPtr<FRHITexture> Texture;
+		FTimespan				  Time;
+		FTimespan				  Duration;
+		FIntPoint				  TotalSize;
+		FIntPoint				  DisplaySize;
+	};
+
+	class FPixelStreaming2TextureSamplePool : public TMediaObjectPool<FPixelStreaming2TextureSample>
+	{
+	};
+
+	class FPixelStreaming2AudioSample : public IMediaAudioSample, public IMediaPoolable
+	{
+	public:
+		FPixelStreaming2AudioSample()
+			: Channels(0)
+			, Duration(FTimespan::Zero())
+			, SampleRate(0)
+			, Time(FTimespan::Zero())
+		{
+		}
+		virtual ~FPixelStreaming2AudioSample() = default;
+
+		void Initialize(
+			const uint8* InBuffer,
+			uint32		 InSize,
+			uint32		 InChannels,
+			uint32		 InSampleRate,
+			FTimespan	 InTime,
+			FTimespan	 InDuration)
+		{
+			check(InBuffer && InSize > 0);
+
+			Buffer.Reset(InSize);
+			Buffer.Append(InBuffer, InSize);
+
+			Channels = InChannels;
+			Duration = InDuration;
+			SampleRate = InSampleRate;
+			Time = InTime;
+		}
+
+		// Begin IMediaAudioSample
+		virtual const void* GetBuffer() override
+		{
+			return Buffer.GetData();
+		}
+
+		virtual uint32 GetChannels() const override
+		{
+			return Channels;
+		}
+
+		virtual FTimespan GetDuration() const override
+		{
+			return Duration;
+		}
+
+		virtual EMediaAudioSampleFormat GetFormat() const override
+		{
+			return EMediaAudioSampleFormat::Int16;
+		}
+
+		virtual uint32 GetFrames() const override
+		{
+			return Buffer.Num() / (Channels * sizeof(int16));
+		}
+
+		virtual uint32 GetSampleRate() const override
+		{
+			return SampleRate;
+		}
+
+		virtual FMediaTimeStamp GetTime() const override
+		{
+			return FMediaTimeStamp(Time);
+		}
+
+		const TArray<uint8>& GetDataBuffer() const
+		{
+			return Buffer;
+		}
+		// End IMediaAudioSample
+
+	private:
+		/** The sample's data buffer. */
+		TArray<uint8> Buffer;
+
+		/** Number of audio channels. */
+		uint32 Channels;
+
+		/** The duration for which the sample is valid. */
+		FTimespan Duration;
+
+		/** Audio sample rate (in samples per second). */
+		uint32 SampleRate;
+
+		/** Presentation time for which the sample was generated. */
+		FTimespan Time;
+	};
+
+	class FPixelStreaming2AudioSamplePool : public TMediaObjectPool<FPixelStreaming2AudioSample>
+	{
+	};
+
+	uint32_t FPixelStreaming2RTCStreamPlayer::PlayerId = 0;
+
+	FPixelStreaming2RTCStreamPlayer::FPixelStreaming2RTCStreamPlayer()
+		: SelectedAudioTrack(INDEX_NONE)
+		, SelectedVideoTrack(INDEX_NONE)
+		, Samples(MakeShared<FMediaSamples, ESPMode::ThreadSafe>())
+		, AudioSamplePool(new FPixelStreaming2AudioSamplePool)
+		, VideoSamplePool(new FPixelStreaming2TextureSamplePool)
+	{
+		PlayerName = FString::Printf(TEXT("FPixelStreaming2RTCStreamPlayer%d"), PlayerId++);
+		OnStatsReadyHandle = FPixelStreaming2RTCModule::GetModule()->GetStatsCollector()->OnStatsReady.AddRaw(this, &FPixelStreaming2RTCStreamPlayer::OnStatsReady);
+	}
+
+	FPixelStreaming2RTCStreamPlayer::~FPixelStreaming2RTCStreamPlayer()
+	{
+		FPixelStreaming2RTCModule::GetModule()->GetStatsCollector()->OnStatsReady.Remove(OnStatsReadyHandle);
+	}
+
+	bool FPixelStreaming2RTCStreamPlayer::Open(const FString& Url, const IMediaOptions* Options)
+	{
+		StreamURL = Url;
+
+		TOptional<FString> UrlStreamerId = FGenericPlatformHttp::GetUrlParameter(Url, TEXT("StreamerId"));
+		if (UrlStreamerId.IsSet())
+		{
+			TargetStreamerId = *UrlStreamerId;
+		}
+
+		TSharedPtr<FPixelStreaming2RTCStreamPlayer> Player = AsShared();
+		if (!SessionObserver)
+		{
+			SessionObserver = MakeRefCount<FEpicRtcSessionObserver>(TObserver<IPixelStreaming2SessionObserver>(Player));
+		}
+
+		if (!RoomObserver)
+		{
+			RoomObserver = MakeRefCount<FEpicRtcRoomObserver>(TObserver<IPixelStreaming2RoomObserver>(Player));
+		}
+
+		if (!AudioTrackObserverFactory)
+		{
+			AudioTrackObserverFactory = MakeRefCount<FEpicRtcAudioTrackObserverFactory>(TObserver<IPixelStreaming2AudioTrackObserver>(Player));
+		}
+
+		if (!VideoTrackObserverFactory)
+		{
+			VideoTrackObserverFactory = MakeRefCount<FEpicRtcVideoTrackObserverFactory>(TObserver<IPixelStreaming2VideoTrackObserver>(Player));
+		}
+
+		if (!DataTrackObserverFactory)
+		{
+			DataTrackObserverFactory = MakeRefCount<FEpicRtcDataTrackObserverFactory>(TObserver<IPixelStreaming2DataTrackObserver>(Player));
+		}
+
+		if (!EpicRtcConference)
+		{
+			EpicRtcConference = FPixelStreaming2RTCModule::GetModule()->GetEpicRtcConference();
+		}
+
+		FUtf8String Utf8PlayerName = *PlayerName;
+		FUtf8String Utf8Url = *StreamURL;
+		FUtf8String Utf8ConnectionUrl = Utf8Url + (Utf8Url.Contains(TEXT("?")) ? TEXT("&") : TEXT("?")) + TEXT("isStreamer=false");
+
+		EpicRtcSessionConfig SessionConfig = {
+			._id = { ._ptr = reinterpret_cast<const char*>(*Utf8PlayerName), ._length = static_cast<uint64_t>(Utf8PlayerName.Len()) },
+			._url = { ._ptr = reinterpret_cast<const char*>(*Utf8ConnectionUrl), ._length = static_cast<uint64_t>(Utf8ConnectionUrl.Len()) },
+			._observer = SessionObserver.GetReference()
+		};
+
+		if (TRefCountPtr<EpicRtcConferenceInterface> SafeConference = EpicRtcConference; SafeConference)
+		{
+			if (const EpicRtcErrorCode Result = SafeConference->CreateSession(SessionConfig, EpicRtcSession.GetInitReference()); Result != EpicRtcErrorCode::Ok)
+			{
+				UE_LOGFMT(LogPixelStreaming2RTC, Error, "Failed to create EpicRtc session, CreateSession returned: {0}", ToString(Result));
+				return false;
+			}
+
+			if (const EpicRtcErrorCode Result = EpicRtcSession->Connect(); Result != EpicRtcErrorCode::Ok)
+			{
+				UE_LOGFMT(LogPixelStreaming2RTC, Error, "Failed to connect EpicRtcSession. Connect returned: {0}", ToString(Result));
+				return false;
+			}
+		}
+		else
+		{
+			return false;
+		}
+
+		return true;
+	}
+
+	bool FPixelStreaming2RTCStreamPlayer::Open(const TSharedRef<FArchive, ESPMode::ThreadSafe>& Archive, const FString& OriginalUrl, const IMediaOptions* Options)
+	{
+		// TODO! We don't support this Open method
+		return false;
+	}
+
+	void FPixelStreaming2RTCStreamPlayer::Close()
+	{
+		TRefCountPtr<EpicRtcSessionInterface> SafeSession = EpicRtcSession;
+		if (!SafeSession)
+		{
+			UE_LOGFMT(LogPixelStreaming2RTC, Error, "Failed to close. EpicRtcSession does not exist!");
+			return;
+		}
+
+		// NOTE: It is imperative we null out the video sink before we remove room and session. If the sink is still alive
+		// during session destruction, webrtc will flush incoming frames (leading to OnVideoTrackFrame) and the engine will
+		// lock inside AVCodecs on a RHI fence. With the sink nulled, OnVideoTrackFrame will early exit
+		if (VideoSink)
+		{
+			if (TSharedPtr<FEpicRtcVideoSink> SafeVideoSink = VideoSink; SafeVideoSink)
+			{
+				SafeVideoSink->RemoveVideoConsumer(TWeakPtrVariant<IPixelStreaming2VideoConsumer>(this));
+			}
+
+			VideoSink = nullptr;
+			RemoteVideoTrack = 0;
+		}
+
+		if (AudioSink)
+		{
+			if (TSharedPtr<FEpicRtcAudioSink> SafeAudioSink = AudioSink; SafeAudioSink)
+			{
+				SafeAudioSink->RemoveAudioConsumer(TWeakPtrVariant<IPixelStreaming2AudioConsumer>(this));
+			}
+
+			AudioSink = nullptr;
+			RemoteAudioTrack = 0;
+		}
+
+		if (TRefCountPtr<EpicRtcRoomInterface> SafeRoom = EpicRtcRoom; SafeRoom)
+		{
+			FUtf8String Utf8TargetStreamerId = *TargetStreamerId;
+
+			SafeRoom->Leave();
+
+			SafeSession->RemoveRoom({ ._ptr = reinterpret_cast<const char*>(*Utf8TargetStreamerId), ._length = static_cast<uint64_t>(Utf8TargetStreamerId.Len()) });
+		}
+
+		FUtf8String Reason = TEXT("FPixelStreaming2RTCStreamPlayer::Close");
+		if (const EpicRtcErrorCode Result = SafeSession->Disconnect(ToEpicRtcStringView(Reason)); Result != EpicRtcErrorCode::Ok)
+		{
+			UE_LOGFMT(LogPixelStreaming2RTC, Error, "Failed to disconnect EpicRtcSession. Disconnect returned {0}", ToString(Result));
+			return;
+		}
+
+		if (TRefCountPtr<EpicRtcConferenceInterface> SafeConference = EpicRtcConference; SafeConference)
+		{
+			FUtf8String Utf8PlayerName = *PlayerName;
+			EpicRtcConference->RemoveSession({ ._ptr = reinterpret_cast<const char*>(*Utf8PlayerName), ._length = static_cast<uint64_t>(Utf8PlayerName.Len()) });
+		}
+	}
+
+	FString FPixelStreaming2RTCStreamPlayer::GetUrl() const
+	{
+		return StreamURL;
+	}
+
+	IMediaSamples& FPixelStreaming2RTCStreamPlayer::GetSamples()
+	{
+		return *Samples.Get();
+	}
+
+	IMediaTracks& FPixelStreaming2RTCStreamPlayer::GetTracks()
+	{
+		return *this;
+	}
+
+	IMediaCache& FPixelStreaming2RTCStreamPlayer::GetCache()
+	{
+		return *this;
+	}
+
+	IMediaControls& FPixelStreaming2RTCStreamPlayer::GetControls()
+	{
+		return *this;
+	}
+
+	FString FPixelStreaming2RTCStreamPlayer::GetInfo() const
+	{
+		return TEXT("FPixelStreaming2RTCStreamPlayer information not implemented yet");
+	}
+
+	FGuid FPixelStreaming2RTCStreamPlayer::GetPlayerPluginGUID() const
+	{
+		FPixelStreaming2RTCModule* ModulePtr = FPixelStreaming2RTCModule::GetModule();
+		check(ModulePtr);
+
+		return ModulePtr->GetPlayerPluginGUID();
+	}
+
+	FString FPixelStreaming2RTCStreamPlayer::GetStats() const
+	{
+		return TEXT("FPixelStreaming2RTCStreamPlayer stats information not implemented yet");
+	}
+
+	IMediaView& FPixelStreaming2RTCStreamPlayer::GetView()
+	{
+		return *this;
+	}
+
+	void FPixelStreaming2RTCStreamPlayer::TickFetch(FTimespan DeltaTime, FTimespan Timecode)
+	{
+	}
+
+	void FPixelStreaming2RTCStreamPlayer::OnSessionStateUpdate(const EpicRtcSessionState StateUpdate)
+	{
+		switch (StateUpdate)
+		{
+			case EpicRtcSessionState::New:
+			case EpicRtcSessionState::Pending:
+			case EpicRtcSessionState::Connected:
+			case EpicRtcSessionState::Disconnected:
+			case EpicRtcSessionState::Failed:
+			case EpicRtcSessionState::Exiting:
+				SessionState = StateUpdate;
+				break;
+			default:
+				UE_LOGFMT(LogPixelStreaming2RTC, Warning, "OnSessionStateUpdate received unknown EpicRtcSessionState: {0}", static_cast<int>(StateUpdate));
+				break;
+		}
+	}
+
+	void FPixelStreaming2RTCStreamPlayer::OnSessionErrorUpdate(const EpicRtcErrorCode ErrorUpdate)
+	{
+	}
+
+	void FPixelStreaming2RTCStreamPlayer::OnSessionRoomsAvailableUpdate(EpicRtcStringArrayInterface* RoomsList)
+	{
+		TArray<FString> Streamers;
+		for (uint64_t i = 0; i < RoomsList->Size(); i++)
+		{
+			EpicRtcStringInterface* RoomName = RoomsList->Get()[i];
+			Streamers.Add(FString{ RoomName->Get(), static_cast<int32>(RoomName->Length()) });
+		}
+
+		if (Streamers.IsEmpty())
+		{
+			UE_LOGFMT(LogPixelStreaming2RTC, Warning, "Streamer list empty!");
+			return;
+		}
+
+		// No specific streamer passed in so subscribe to the first one
+		if (TargetStreamerId.IsEmpty())
+		{
+			TargetStreamerId = Streamers[0];
+		}
+		else
+		{
+			if (!Streamers.Contains(TargetStreamerId))
+			{
+				UE_LOGFMT(LogPixelStreaming2RTC, Warning, "Streamer list [{0}] doesn't contain target streamer [{1}]!", FString::Join(Streamers, TEXT(",")), TargetStreamerId);
+				return;
+			}
+		}
+
+		if (SessionState != EpicRtcSessionState::Connected)
+		{
+			UE_LOGFMT(LogPixelStreaming2RTC, Error, "Failed to create subscribe to streamer. EpicRtc session isn't connected");
+			return;
+		}
+
+		EpicRtcConnectionConfig ConnectionConfig = {
+			// TODO (Migration): RTCP-7032 This info usually comes from the OnConfig signalling message
+			._iceServers = EpicRtcIceServerSpan{ ._ptr = nullptr, ._size = 0 },
+			._iceConnectionPolicy = EpicRtcIcePolicy::All,
+			._disableTcpCandidates = false
+		};
+
+		FUtf8String		  Utf8TargetStreamerId = *TargetStreamerId;
+		EpicRtcRoomConfig RoomConfig = {
+			._id = EpicRtcStringView{ ._ptr = reinterpret_cast<const char*>(*Utf8TargetStreamerId), ._length = static_cast<uint64_t>(Utf8TargetStreamerId.Len()) },
+			._connectionConfig = ConnectionConfig,
+			._ticket = EpicRtcStringView{ ._ptr = nullptr, ._length = 0 },
+			._observer = RoomObserver,
+			._audioTrackObserverFactory = AudioTrackObserverFactory,
+			._dataTrackObserverFactory = DataTrackObserverFactory,
+			._videoTrackObserverFactory = VideoTrackObserverFactory
+		};
+
+		if (const EpicRtcErrorCode Result = EpicRtcSession->CreateRoom(RoomConfig, EpicRtcRoom.GetInitReference()); Result != EpicRtcErrorCode::Ok)
+		{
+			UE_LOGFMT(LogPixelStreaming2RTC, Error, "Failed to create EpicRtc room. CreateRoom returned: {0}", ToString(Result));
+			return;
+		}
+
+		EpicRtcRoom->Join();
+
+		// Create a stats collector so we can receive stats from the subscribed streamer
+		StatsCollector = FRTCStatsCollector::Create(*Utf8TargetStreamerId);
+	}
+	void FPixelStreaming2RTCStreamPlayer::OnRoomStateUpdate(const EpicRtcRoomState State)
+	{
+	}
+
+	void FPixelStreaming2RTCStreamPlayer::OnRoomJoinedUpdate(EpicRtcParticipantInterface* Participant)
+	{
+		FString ParticipantId = ToString(Participant->GetId());
+		UE_LOGFMT(LogPixelStreaming2RTC, Log, "Player ({0}) joined", ParticipantId);
+	}
+
+	void FPixelStreaming2RTCStreamPlayer::OnRoomLeftUpdate(const EpicRtcStringView ParticipantId)
+	{
+	}
+
+	void FPixelStreaming2RTCStreamPlayer::OnAudioTrackUpdate(EpicRtcParticipantInterface* Participant, EpicRtcAudioTrackInterface* InAudioTrack)
+	{
+		FString ParticipantId = ToString(Participant->GetId());
+		FString AudioTrackId = ToString(InAudioTrack->GetId());
+
+		UE_LOGFMT(LogPixelStreaming2RTC, VeryVerbose, "FPixelStreaming2RTCStreamPlayer::OnAudioTrackUpdate(Participant [{0}], AudioTrack [{1}])", ParticipantId, AudioTrackId);
+		if (InAudioTrack->IsRemote())
+		{
+			// We received a remote track. We should now create a sink to handle receiving the frames.
+			// NOTE: We pass in nullptr as the track because if we store the track on the sink, EpicRtc will be unable to destroy it
+			// and webrtc will try to flush remaining frames during session removal.
+			AudioSink = FEpicRtcAudioSink::Create(nullptr);
+			RemoteAudioTrack = reinterpret_cast<uintptr_t>(InAudioTrack); // Keep track of which remote track we're receiving audio from as we only support display one track
+			// As EpicRtc work can happen on different threads, we always need to call functions on a copy of the AudioSink to ensure it doesn't get deleted out from underneath us
+			if (TSharedPtr<FEpicRtcAudioSink> SafeAudioSink = AudioSink; SafeAudioSink)
+			{
+				SafeAudioSink->AddAudioConsumer(TWeakPtrVariant<IPixelStreaming2AudioConsumer>(this));
+			}
+		}
+	}
+
+	void FPixelStreaming2RTCStreamPlayer::OnVideoTrackUpdate(EpicRtcParticipantInterface* Participant, EpicRtcVideoTrackInterface* InVideoTrack)
+	{
+
+		FString ParticipantId = ToString(Participant->GetId());
+		FString VideoTrackId = ToString(InVideoTrack->GetId());
+
+		if (VideoTrackId == TEXT("probator"))
+		{
+			// probator is fake sfu news
+			return;
+		}
+
+		UE_LOGFMT(LogPixelStreaming2RTC, VeryVerbose, "FPixelStreaming2RTCStreamPlayer::OnVideoTrackUpdate(Participant [{0}], VideoTrack [{1}])", ParticipantId, VideoTrackId);
+		if (InVideoTrack->IsRemote())
+		{
+			// We received a remote track. We should now create a sink to handle receiving the frames.
+			// NOTE: We pass in nullptr as the track because if we store the track on the sink, EpicRtc will be unable to destroy it
+			// and webrtc will try to flush remaining frames during session removal.
+			VideoSink = FEpicRtcVideoSink::Create(nullptr);
+			RemoteVideoTrack = reinterpret_cast<uintptr_t>(InVideoTrack); // Keep track of which remote track we're receiving video from as we only support display one track
+			// As EpicRtc work can happen on different threads, we always need to call functions on a copy of the VideoSink to ensure it doesn't get deleted out from underneath us
+			if (TSharedPtr<FEpicRtcVideoSink> SafeVideoSink = VideoSink; SafeVideoSink)
+			{
+				SafeVideoSink->AddVideoConsumer(TWeakPtrVariant<IPixelStreaming2VideoConsumer>(this));
+			}
+		}
+	}
+
+	void FPixelStreaming2RTCStreamPlayer::OnDataTrackUpdate(EpicRtcParticipantInterface* Participant, EpicRtcDataTrackInterface* DataTrack)
+	{
+	}
+
+	EpicRtcSdpInterface* FPixelStreaming2RTCStreamPlayer::OnLocalSdpUpdate(EpicRtcParticipantInterface* Participant, EpicRtcSdpInterface* Sdp)
+	{
+		return Sdp;
+	}
+
+	EpicRtcSdpInterface* FPixelStreaming2RTCStreamPlayer::OnRemoteSdpUpdate(EpicRtcParticipantInterface* Participant, EpicRtcSdpInterface* Sdp)
+	{
+		return Sdp;
+	}
+
+	void FPixelStreaming2RTCStreamPlayer::OnRoomErrorUpdate(const EpicRtcErrorCode Error)
+	{
+	}
+
+	void FPixelStreaming2RTCStreamPlayer::OnAudioTrackMuted(EpicRtcAudioTrackInterface* InAudioTrack, EpicRtcBool bIsMuted)
+	{
+		if (!AudioSink || RemoteAudioTrack != reinterpret_cast<uintptr_t>(InAudioTrack))
+		{
+			return;
+		}
+
+		if (TSharedPtr<FEpicRtcAudioSink> SafeAudioSink = AudioSink; SafeAudioSink)
+		{
+			SafeAudioSink->SetMuted(static_cast<bool>(bIsMuted));
+		}
+	}
+
+	void FPixelStreaming2RTCStreamPlayer::OnAudioTrackFrame(EpicRtcAudioTrackInterface* InAudioTrack, const EpicRtcAudioFrame& Frame)
+	{
+		if (!AudioSink || RemoteAudioTrack != reinterpret_cast<uintptr_t>(InAudioTrack))
+		{
+			return;
+		}
+
+		if (TSharedPtr<FEpicRtcAudioSink> SafeAudioSink = AudioSink; SafeAudioSink)
+		{
+			SafeAudioSink->OnAudioData(Frame._data, Frame._length, Frame._format._numChannels, Frame._format._sampleRate);
+		}
+	}
+
+	void FPixelStreaming2RTCStreamPlayer::OnAudioTrackRemoved(EpicRtcAudioTrackInterface* InAudioTrack)
+	{
+	}
+
+	void FPixelStreaming2RTCStreamPlayer::OnAudioTrackState(EpicRtcAudioTrackInterface* InAudioTrack, const EpicRtcTrackState State)
+	{
+	}
+
+	void FPixelStreaming2RTCStreamPlayer::OnVideoTrackMuted(EpicRtcVideoTrackInterface* InVideoTrack, EpicRtcBool bIsMuted)
+	{
+		if (!VideoSink || RemoteVideoTrack != reinterpret_cast<uintptr_t>(InVideoTrack))
+		{
+			return;
+		}
+
+		if (TSharedPtr<FEpicRtcVideoSink> SafeVideoSink = VideoSink; SafeVideoSink)
+		{
+			SafeVideoSink->SetMuted(static_cast<bool>(bIsMuted));
+		}
+	}
+
+	void FPixelStreaming2RTCStreamPlayer::OnVideoTrackFrame(EpicRtcVideoTrackInterface* InVideoTrack, const EpicRtcVideoFrame& Frame)
+	{
+		if (!VideoSink || RemoteVideoTrack != reinterpret_cast<uintptr_t>(InVideoTrack))
+		{
+			return;
+		}
+
+		if (TSharedPtr<FEpicRtcVideoSink> SafeVideoSink = VideoSink; SafeVideoSink)
+		{
+			SafeVideoSink->OnEpicRtcFrame(Frame);
+		}
+	}
+
+	void FPixelStreaming2RTCStreamPlayer::OnVideoTrackRemoved(EpicRtcVideoTrackInterface* InVideoTrack)
+	{
+	}
+
+	void FPixelStreaming2RTCStreamPlayer::OnVideoTrackState(EpicRtcVideoTrackInterface* InVideoTrack, const EpicRtcTrackState State)
+	{
+	}
+
+	void FPixelStreaming2RTCStreamPlayer::OnVideoTrackEncodedFrame(EpicRtcVideoTrackInterface* InVideoTrack, const EpicRtcEncodedVideoFrame& EncodedFrame)
+	{
+	}
+
+	EpicRtcBool FPixelStreaming2RTCStreamPlayer::Enabled() const
+	{
+		return true;
+	}
+
+	void FPixelStreaming2RTCStreamPlayer::OnDataTrackState(EpicRtcDataTrackInterface* DataTrack, const EpicRtcTrackState State)
+	{
+	}
+
+	void FPixelStreaming2RTCStreamPlayer::OnDataTrackMessage(EpicRtcDataTrackInterface* DataTrack)
+	{
+	}
+
+	void FPixelStreaming2RTCStreamPlayer::OnDataTrackError(EpicRtcDataTrackInterface* DataTrack, const EpicRtcErrorCode Error)
+	{
+	}
+
+	void FPixelStreaming2RTCStreamPlayer::ConsumeFrame(FTextureRHIRef Frame)
+	{
+		TSharedRef<FPixelStreaming2TextureSample, ESPMode::ThreadSafe> VideoSample = VideoSamplePool->AcquireShared();
+		// Fake timestamp so we always display the frame
+		VideoSample->Initialize(Frame, Frame->GetSizeXY(), Frame->GetSizeXY(), FTimespan::Zero(), FTimespan(1));
+
+		Samples->AddVideo(VideoSample);
+	}
+
+	void FPixelStreaming2RTCStreamPlayer::ConsumeRawPCM(const int16_t* AudioData, int InSampleRate, size_t NChannels, size_t NFrames)
+	{
+		TSharedRef<FPixelStreaming2AudioSample, ESPMode::ThreadSafe> AudioSample = AudioSamplePool->AcquireShared();
+
+		TArray<uint8> AudioBuffer;
+
+		size_t NumSamples = NFrames * NChannels;
+		size_t DataSize = NumSamples * sizeof(int16);
+		AudioBuffer.SetNumUninitialized(DataSize);
+
+		FMemory::Memcpy(AudioBuffer.GetData(), AudioData, DataSize);
+
+		FTimespan Duration = (NumSamples * ETimespan::TicksPerSecond) / InSampleRate;
+		AudioSample->Initialize(AudioBuffer.GetData(), DataSize, NChannels, InSampleRate, FTimespan::Zero(), Duration);
+
+		Samples->AddAudio(AudioSample);
+	}
+
+	void FPixelStreaming2RTCStreamPlayer::OnConsumerAdded()
+	{
+	}
+
+	void FPixelStreaming2RTCStreamPlayer::OnConsumerRemoved()
+	{
+	}
+
+	bool FPixelStreaming2RTCStreamPlayer::GetAudioTrackFormat(int32 TrackIndex, int32 FormatIndex, FMediaAudioTrackFormat& OutFormat) const
+	{
+		return false;
+	}
+
+	int32 FPixelStreaming2RTCStreamPlayer::GetNumTracks(EMediaTrackType TrackType) const
+	{
+		// We support only video and audio tracks
+		return (TrackType == EMediaTrackType::Audio || TrackType == EMediaTrackType::Video) ? 1 : 0;
+	}
+
+	int32 FPixelStreaming2RTCStreamPlayer::GetNumTrackFormats(EMediaTrackType TrackType, int32 TrackIndex) const
+	{
+		return 1;
+	}
+
+	int32 FPixelStreaming2RTCStreamPlayer::GetSelectedTrack(EMediaTrackType TrackType) const
+	{
+		if (TrackType == EMediaTrackType::Audio)
+		{
+			return SelectedAudioTrack;
+		}
+		else if (TrackType == EMediaTrackType::Video)
+		{
+			return SelectedVideoTrack;
+		}
+		else
+		{
+			return INDEX_NONE;
+		}
+	}
+
+	FText FPixelStreaming2RTCStreamPlayer::GetTrackDisplayName(EMediaTrackType TrackType, int32 TrackIndex) const
+	{
+		if (TrackType == EMediaTrackType::Audio)
+		{
+			return FText::FromString(TEXT("PixelStreaming2Audio"));
+		}
+		else if (TrackType == EMediaTrackType::Video)
+		{
+			return FText::FromString(TEXT("PixelStreaming2Video"));
+		}
+		else
+		{
+			return FText::FromString(TEXT("Unknown"));
+		}
+	}
+
+	int32 FPixelStreaming2RTCStreamPlayer::GetTrackFormat(EMediaTrackType TrackType, int32 TrackIndex) const
+	{
+		return 0;
+	}
+
+	FString FPixelStreaming2RTCStreamPlayer::GetTrackLanguage(EMediaTrackType TrackType, int32 TrackIndex) const
+	{
+		static FString Language(TEXT("Default"));
+		return Language;
+	}
+
+	FString FPixelStreaming2RTCStreamPlayer::GetTrackName(EMediaTrackType TrackType, int32 TrackIndex) const
+	{
+		if (TrackType == EMediaTrackType::Audio)
+		{
+			return TEXT("PixelStreaming2Audio");
+		}
+		else if (TrackType == EMediaTrackType::Video)
+		{
+			return TEXT("PixelStreaming2Video");
+		}
+		else
+		{
+			return TEXT("Unknown");
+		}
+	}
+
+	bool FPixelStreaming2RTCStreamPlayer::GetVideoTrackFormat(int32 TrackIndex, int32 FormatIndex, FMediaVideoTrackFormat& OutFormat) const
+	{
+		return false;
+	}
+
+	bool FPixelStreaming2RTCStreamPlayer::SelectTrack(EMediaTrackType TrackType, int32 TrackIndex)
+	{
+		bool bHasChanged = false;
+
+		if (TrackType == EMediaTrackType::Audio)
+		{
+			bHasChanged = SelectedAudioTrack != TrackIndex;
+			SelectedAudioTrack = TrackIndex;
+		}
+		else if (TrackType == EMediaTrackType::Video)
+		{
+			bHasChanged = SelectedVideoTrack != TrackIndex;
+			SelectedVideoTrack = TrackIndex;
+		}
+
+		return true;
+	}
+
+	bool FPixelStreaming2RTCStreamPlayer::SetTrackFormat(EMediaTrackType TrackType, int32 TrackIndex, int32 FormatIndex)
+	{
+		if (FormatIndex == 0)
+		{
+			return true;
+		}
+		else
+		{
+			// We only support one track format
+			return false;
+		}
+	}
+
+	bool FPixelStreaming2RTCStreamPlayer::CanControl(EMediaControl Control) const
+	{
+		// We don't support any media controls
+		return false;
+	}
+
+	FTimespan FPixelStreaming2RTCStreamPlayer::GetDuration() const
+	{
+		return FTimespan::Zero();
+	}
+
+	float FPixelStreaming2RTCStreamPlayer::GetRate() const
+	{
+		// We only support 1x playback
+		return 1.f;
+	}
+
+	EMediaState FPixelStreaming2RTCStreamPlayer::GetState() const
+	{
+		// As we don't support media controls, we're always playing
+		return EMediaState::Playing;
+	}
+
+	EMediaStatus FPixelStreaming2RTCStreamPlayer::GetStatus() const
+	{
+		return EMediaStatus::None;
+	}
+
+	TRangeSet<float> FPixelStreaming2RTCStreamPlayer::GetSupportedRates(EMediaRateThinning Thinning) const
+	{
+		TRangeSet<float> Result;
+		// We only support 1x playback
+		Result.Add(TRange<float>(1.0f));
+		return Result;
+	}
+
+	FTimespan FPixelStreaming2RTCStreamPlayer::GetTime() const
+	{
+		return FTimespan(FPlatformTime::Cycles64());
+	}
+
+	bool FPixelStreaming2RTCStreamPlayer::IsLooping() const
+	{
+		return false;
+	}
+
+	bool FPixelStreaming2RTCStreamPlayer::Seek(const FTimespan& Time)
+	{
+		// We don't support seeking of the stream
+		return false;
+	}
+
+	bool FPixelStreaming2RTCStreamPlayer::SetLooping(bool Looping)
+	{
+		return false;
+	}
+
+	bool FPixelStreaming2RTCStreamPlayer::SetRate(float Rate)
+	{
+		// We don't support setting any rates of the stream but must return true to succeed with initialization
+		return Rate == 1.f;
+	}
+
+	void FPixelStreaming2RTCStreamPlayer::OnStatsReady(const FString& PeerId, const EpicRtcConnectionStats& ConnectionStats)
+	{
+		if (!StatsCollector)
+		{
+			return;
+		}
+
+		if (PeerId != TargetStreamerId)
+		{
+			return;
+		}
+
+		StatsCollector->Process(ConnectionStats);
+	}
+} // namespace UE::PixelStreaming2
\ No newline at end of file
diff --git a/Engine/Plugins/Media/PixelStreaming2/Source/PixelStreaming2RTC/Private/PixelStreaming2RTCPlayer.h b/Engine/Plugins/Media/PixelStreaming2/Source/PixelStreaming2RTC/Private/PixelStreaming2RTCPlayer.h
new file mode 100644
--- /dev/null
+++ b/Engine/Plugins/Media/PixelStreaming2/Source/PixelStreaming2RTC/Private/PixelStreaming2RTCPlayer.h
@@ -0,0 +1,201 @@
+// Copyright Epic Games, Inc. All Rights Reserved.
+
+#pragma once
+
+#include "Containers/Utf8String.h"
+#include "Containers/UnrealString.h"
+#include "EpicRtcAudioSink.h"
+#include "EpicRtcAudioTrackObserver.h"
+#include "EpicRtcAudioTrackObserverFactory.h"
+#include "EpicRtcDataTrackObserver.h"
+#include "EpicRtcDataTrackObserverFactory.h"
+#include "EpicRtcVideoSink.h"
+#include "EpicRtcVideoTrackObserver.h"
+#include "EpicRtcVideoTrackObserverFactory.h"
+#include "EpicRtcRoomObserver.h"
+#include "EpicRtcSessionObserver.h"
+#include "IMediaCache.h"
+#include "IMediaControls.h"
+#include "IMediaOptions.h"
+#include "IMediaPlayer.h"
+#include "IMediaTracks.h"
+#include "IMediaView.h"
+#include "IPixelStreaming2AudioConsumer.h"
+#include "IPixelStreaming2VideoConsumer.h"
+#include "MediaSamples.h"
+#include "MediaSource.h"
+#include "RTCStatsCollector.h"
+
+#include "epic_rtc/core/conference.h"
+#include "epic_rtc/core/room.h"
+#include "epic_rtc/core/session.h"
+
+namespace UE::PixelStreaming2
+{
+	class FPixelStreaming2RTCStreamPlayer :
+		public IMediaPlayer,
+		public IMediaTracks,
+		public IMediaCache,
+		public IMediaControls,
+		public IMediaView,
+		public IPixelStreaming2SessionObserver,
+		public IPixelStreaming2RoomObserver,
+		public IPixelStreaming2AudioTrackObserver,
+		public IPixelStreaming2DataTrackObserver,
+		public IPixelStreaming2VideoTrackObserver,
+		public IPixelStreaming2AudioConsumer,
+		public IPixelStreaming2VideoConsumer,
+		public TSharedFromThis<FPixelStreaming2RTCStreamPlayer>
+	{
+	public:
+		FPixelStreaming2RTCStreamPlayer();
+		virtual ~FPixelStreaming2RTCStreamPlayer();
+
+		// Begin IMediaPlayer
+		virtual bool			Open(const FString& Url, const IMediaOptions* Options) override;
+		virtual bool			Open(const TSharedRef<FArchive, ESPMode::ThreadSafe>& Archive, const FString& OriginalUrl, const IMediaOptions* Options) override;
+		virtual void			Close() override;
+		virtual FString			GetUrl() const override;
+		virtual IMediaSamples&	GetSamples() override;
+		virtual IMediaTracks&	GetTracks() override;
+		virtual IMediaCache&	GetCache() override;
+		virtual IMediaControls& GetControls() override;
+		virtual FString			GetInfo() const override;
+		virtual FGuid			GetPlayerPluginGUID() const override;
+		virtual FString			GetStats() const override;
+		virtual IMediaView&		GetView() override;
+		virtual void			TickFetch(FTimespan DeltaTime, FTimespan Timecode) override;
+		// End IMediaPlayer
+
+		// Begin IMediaTracks
+		virtual bool	GetAudioTrackFormat(int32 TrackIndex, int32 FormatIndex, FMediaAudioTrackFormat& OutFormat) const override;
+		virtual int32	GetNumTracks(EMediaTrackType TrackType) const override;
+		virtual int32	GetNumTrackFormats(EMediaTrackType TrackType, int32 TrackIndex) const override;
+		virtual int32	GetSelectedTrack(EMediaTrackType TrackType) const override;
+		virtual FText	GetTrackDisplayName(EMediaTrackType TrackType, int32 TrackIndex) const override;
+		virtual int32	GetTrackFormat(EMediaTrackType TrackType, int32 TrackIndex) const override;
+		virtual FString GetTrackLanguage(EMediaTrackType TrackType, int32 TrackIndex) const override;
+		virtual FString GetTrackName(EMediaTrackType TrackType, int32 TrackIndex) const override;
+		virtual bool	GetVideoTrackFormat(int32 TrackIndex, int32 FormatIndex, FMediaVideoTrackFormat& OutFormat) const override;
+		virtual bool	SelectTrack(EMediaTrackType TrackType, int32 TrackIndex) override;
+		virtual bool	SetTrackFormat(EMediaTrackType TrackType, int32 TrackIndex, int32 FormatIndex) override;
+		// End IMediaTracks
+
+		// Begin IMediaControls
+		virtual bool			 CanControl(EMediaControl Control) const override;
+		virtual FTimespan		 GetDuration() const override;
+		virtual float			 GetRate() const override;
+		virtual EMediaState		 GetState() const override;
+		virtual EMediaStatus	 GetStatus() const override;
+		virtual TRangeSet<float> GetSupportedRates(EMediaRateThinning Thinning) const override;
+		virtual FTimespan		 GetTime() const override;
+		virtual bool			 IsLooping() const override;
+		virtual bool			 Seek(const FTimespan& Time) override;
+		virtual bool			 SetLooping(bool Looping) override;
+		virtual bool			 SetRate(float Rate) override;
+		// End IMediaControls
+
+	private:
+		// Incremented Id used to construct the PlayerName
+		static uint32_t PlayerId;
+		// A unique identifier for the player instance. Used to create unique EpicRtc sessions per player instance
+		FString PlayerName;
+		// The URL of the stream to connect to
+		FString StreamURL;
+		// The id of the streamer to subscribe to. Optionally passed in through the stream URL. If not present in URL
+		// the first streamer id returned in the streamer list will be the streamer that is subscribed to
+		FString TargetStreamerId;
+
+		// The audio sink used to process remote audio
+		TSharedPtr<FEpicRtcAudioSink> AudioSink;
+		// The remote audio track. We only support one track at this time so we make sure we only receive frames for that track
+		uintptr_t RemoteAudioTrack;
+
+		// The video sink used to process remote audio
+		TSharedPtr<FEpicRtcVideoSink> VideoSink;
+		// The remote video track. We only support one track at this time so we make sure we only receive frames for that track
+		uintptr_t RemoteVideoTrack;
+
+		TSharedPtr<FRTCStatsCollector> StatsCollector;
+		FDelegateHandle				   OnStatsReadyHandle;
+
+		int32 SelectedAudioTrack;
+		int32 SelectedVideoTrack;
+
+		TSharedPtr<FMediaSamples, ESPMode::ThreadSafe> Samples;
+
+		TSharedPtr<class FPixelStreaming2AudioSamplePool>	AudioSamplePool;
+		TSharedPtr<class FPixelStreaming2TextureSamplePool> VideoSamplePool;
+
+	public:
+		// Begin IPixelStreaming2SessionObserver
+		virtual void OnSessionStateUpdate(const EpicRtcSessionState StateUpdate) override;
+		virtual void OnSessionErrorUpdate(const EpicRtcErrorCode ErrorUpdate) override;
+		virtual void OnSessionRoomsAvailableUpdate(EpicRtcStringArrayInterface* RoomsList) override;
+		// End IPixelStreaming2SessionObserver
+
+		// Begin IPixelStreaming2RoomObserver
+		virtual void							   OnRoomStateUpdate(const EpicRtcRoomState State) override;
+		virtual void							   OnRoomJoinedUpdate(EpicRtcParticipantInterface* Participant) override;
+		virtual void							   OnRoomLeftUpdate(const EpicRtcStringView ParticipantId) override;
+		virtual void							   OnAudioTrackUpdate(EpicRtcParticipantInterface* Participant, EpicRtcAudioTrackInterface* AudioTrack) override;
+		virtual void							   OnVideoTrackUpdate(EpicRtcParticipantInterface* Participant, EpicRtcVideoTrackInterface* VideoTrack) override;
+		virtual void							   OnDataTrackUpdate(EpicRtcParticipantInterface* Participant, EpicRtcDataTrackInterface* DataTrack) override;
+		[[nodiscard]] virtual EpicRtcSdpInterface* OnLocalSdpUpdate(EpicRtcParticipantInterface* Participant, EpicRtcSdpInterface* Sdp) override;
+		[[nodiscard]] virtual EpicRtcSdpInterface* OnRemoteSdpUpdate(EpicRtcParticipantInterface* Participant, EpicRtcSdpInterface* Sdp) override;
+		virtual void							   OnRoomErrorUpdate(const EpicRtcErrorCode Error) override;
+		// End IPixelStreaming2RoomObserver
+
+		// Begin IPixelStreaming2AudioTrackObserver
+		virtual void OnAudioTrackMuted(EpicRtcAudioTrackInterface* AudioTrack, EpicRtcBool bIsMuted) override;
+		virtual void OnAudioTrackFrame(EpicRtcAudioTrackInterface* AudioTrack, const EpicRtcAudioFrame& Frame) override;
+		virtual void OnAudioTrackRemoved(EpicRtcAudioTrackInterface* AudioTrack) override;
+		virtual void OnAudioTrackState(EpicRtcAudioTrackInterface* AudioTrack, const EpicRtcTrackState State) override;
+		// End IPixelStreaming2AudioTrackObserver
+
+		// Begin IPixelStreaming2VideoTrackObserver
+		virtual void		OnVideoTrackMuted(EpicRtcVideoTrackInterface* VideoTrack, EpicRtcBool bIsMuted) override;
+		virtual void		OnVideoTrackFrame(EpicRtcVideoTrackInterface* VideoTrack, const EpicRtcVideoFrame& Frame) override;
+		virtual void		OnVideoTrackRemoved(EpicRtcVideoTrackInterface* VideoTrack) override;
+		virtual void		OnVideoTrackState(EpicRtcVideoTrackInterface* VideoTrack, const EpicRtcTrackState State) override;
+		virtual void		OnVideoTrackEncodedFrame(EpicRtcVideoTrackInterface* VideoTrack, const EpicRtcEncodedVideoFrame& EncodedFrame) override;
+		virtual EpicRtcBool Enabled() const override;
+		// End IPixelStreaming2VideoTrackObserver
+
+		// Begin IPixelStreaming2DataTrackObserver
+		virtual void OnDataTrackState(EpicRtcDataTrackInterface* DataTrack, const EpicRtcTrackState State) override;
+		virtual void OnDataTrackMessage(EpicRtcDataTrackInterface* DataTrack) override;
+		virtual void OnDataTrackError(EpicRtcDataTrackInterface* DataTrack, const EpicRtcErrorCode Error) override;
+		// End IPixelStreaming2DataTrackObserver
+
+		// Begin IPixelStreaming2VideoConsumer
+		virtual void ConsumeFrame(FTextureRHIRef Frame) override;
+		virtual void OnConsumerAdded() override;
+		virtual void OnConsumerRemoved() override;
+		// End IPixelStreaming2VideoConsumer
+
+		// Begin IPixelStreaming2AudioConsumer
+		void ConsumeRawPCM(const int16_t* AudioData, int InSampleRate, size_t NChannels, size_t NFrames) override;
+		// NOTE: OnConsumerAdded/Removed is already overridden by the video consumer
+		// End IPixelStreaming2AudioConsumer
+
+		void OnStatsReady(const FString& PeerId, const EpicRtcConnectionStats& ConnectionStats);
+
+	private:
+		EpicRtcSessionState SessionState = EpicRtcSessionState::Disconnected;
+
+		// Begin EpicRtc Classes
+		TRefCountPtr<EpicRtcConferenceInterface> EpicRtcConference;
+		TRefCountPtr<EpicRtcSessionInterface>	 EpicRtcSession;
+		TRefCountPtr<EpicRtcRoomInterface>		 EpicRtcRoom;
+		// End EpicRtc Classes
+
+		// Begin EpicRtc Observers
+		TRefCountPtr<FEpicRtcSessionObserver>			SessionObserver;
+		TRefCountPtr<FEpicRtcRoomObserver>				RoomObserver;
+		TRefCountPtr<FEpicRtcAudioTrackObserverFactory> AudioTrackObserverFactory;
+		TRefCountPtr<FEpicRtcVideoTrackObserverFactory> VideoTrackObserverFactory;
+		TRefCountPtr<FEpicRtcDataTrackObserverFactory>	DataTrackObserverFactory;
+		// End EpicRtc Observers
+	};
+} // namespace UE::PixelStreaming2
\ No newline at end of file
