[PixelStreaming2] Fix: Unify reconnect logic. StreamerReconnectTimer fully handles it

diff --git a/Engine/Plugins/Media/PixelStreaming2/Source/PixelStreaming2/Internal/StreamerReconnectTimer.h b/Engine/Plugins/Media/PixelStreaming2/Source/PixelStreaming2/Internal/StreamerReconnectTimer.h
--- a/Engine/Plugins/Media/PixelStreaming2/Source/PixelStreaming2/Internal/StreamerReconnectTimer.h
+++ b/Engine/Plugins/Media/PixelStreaming2/Source/PixelStreaming2/Internal/StreamerReconnectTimer.h
@@ -10,7 +10,6 @@ namespace UE::PixelStreaming2
 {
 	/**
 	 * An interval timer that triggers reconnects when `IPixelStreaming2Streamer` has it session disconnect.
-	 * This is different than a websocket error reconnect because in this case we were already connected.
 	 * Timer takes a `TWeakPtr<IPixelStreaming2Streamer>`, meaning the streamer object must exist for the timer to do anything.
 	 * Note: this timer runs on the game thread and its reconnect interval is controlled by the CVar
 	 * UPixelStreaming2PluginSettings::CVarSignalingReconnectInterval
@@ -18,10 +17,11 @@ namespace UE::PixelStreaming2
 	class PIXELSTREAMING2_API FStreamerReconnectTimer : public FTickableGameObject
 	{
 	public:
-		FStreamerReconnectTimer();
+		FStreamerReconnectTimer() = default;
 		virtual ~FStreamerReconnectTimer() = default;
 		void Start(TWeakPtr<IPixelStreaming2Streamer> InWeakStreamer);
 		void Stop();
+		void Reset();
 
 		// Begin FTickableGameObject
 		virtual void		Tick(float DeltaTime) override;
@@ -35,6 +35,7 @@ namespace UE::PixelStreaming2
 
 	private:
 		TWeakPtr<IPixelStreaming2Streamer> WeakStreamer;
+
 		bool							   bEnabled = false;
 		uint64							   NumReconnectAttempts = 0;
 		uint64							   LastReconnectCycles = 0;
diff --git a/Engine/Plugins/Media/PixelStreaming2/Source/PixelStreaming2/Private/StreamerReconnectTimer.cpp b/Engine/Plugins/Media/PixelStreaming2/Source/PixelStreaming2/Private/StreamerReconnectTimer.cpp
--- a/Engine/Plugins/Media/PixelStreaming2/Source/PixelStreaming2/Private/StreamerReconnectTimer.cpp
+++ b/Engine/Plugins/Media/PixelStreaming2/Source/PixelStreaming2/Private/StreamerReconnectTimer.cpp
@@ -7,15 +7,9 @@
 
 namespace UE::PixelStreaming2
 {
-
-	FStreamerReconnectTimer::FStreamerReconnectTimer()
-	{
-	}
-
 	void FStreamerReconnectTimer::Start(TWeakPtr<IPixelStreaming2Streamer> InWeakStreamer)
 	{
 		WeakStreamer = InWeakStreamer;
-		NumReconnectAttempts = 0;
 		bEnabled = true;
 	}
 
@@ -24,6 +18,11 @@ namespace UE::PixelStreaming2
 		bEnabled = false;
 	}
 
+	void FStreamerReconnectTimer::Reset()
+	{
+		NumReconnectAttempts = 0;
+	}
+
 	void FStreamerReconnectTimer::Tick(float DeltaTime)
 	{
 		if (IsEngineExitRequested())
@@ -68,6 +67,7 @@ namespace UE::PixelStreaming2
 			{
 				// Maxmimum exceeded so don't attempt it and instead stop the timer
 				Stop();
+				UE_LOGFMT(LogPixelStreaming2, Log, "Exceeded maximum reconnect attempts!", NumReconnectAttempts);
 				OnExceededMaxReconnectAttempts.Broadcast();
 				return;
 			}
diff --git a/Engine/Plugins/Media/PixelStreaming2/Source/PixelStreaming2RTC/Internal/EpicRtcWebsocket.h b/Engine/Plugins/Media/PixelStreaming2/Source/PixelStreaming2RTC/Internal/EpicRtcWebsocket.h
--- a/Engine/Plugins/Media/PixelStreaming2/Source/PixelStreaming2RTC/Internal/EpicRtcWebsocket.h
+++ b/Engine/Plugins/Media/PixelStreaming2/Source/PixelStreaming2RTC/Internal/EpicRtcWebsocket.h
@@ -15,7 +15,7 @@ namespace UE::PixelStreaming2
 	class PIXELSTREAMING2RTC_API FEpicRtcWebsocket : public EpicRtcWebsocketInterface, public FTickableGameObject
 	{
 	public:
-		FEpicRtcWebsocket(bool bKeepAlive = true, TSharedPtr<IWebSocket> WebSocket = nullptr, TFunction<void(void)> OnMaxReconnectAttemptsExceeded = []() {});
+		FEpicRtcWebsocket(bool bKeepAlive = true, TSharedPtr<IWebSocket> WebSocket = nullptr);
 		virtual ~FEpicRtcWebsocket() = default;
 
 		// Begin EpicRtcWebsocketInterface
@@ -37,7 +37,6 @@ namespace UE::PixelStreaming2
 		void OnClosed(int32 StatusCode, const FString& Reason, bool bWasClean);
 		void OnMessage(const FString& Msg);
 		void OnBinaryMessage(const uint8* Data, int32 Length, bool bIsLastFragment);
-		void Reconnect();
 		void KeepAlive();
 
 	private:
@@ -53,12 +52,8 @@ namespace UE::PixelStreaming2
 		TSharedPtr<IWebSocket>							WebSocket;
 		TRefCountPtr<EpicRtcWebsocketObserverInterface> Observer;
 		uint64											LastKeepAliveCycles = 0;
-		uint64											LastReconnectCycles = 0;
-		uint64											NumReconnectAttempts = 0;
 		bool											bSendKeepAlive = false;
-		bool											bReconnectOnError = false;
 		bool											bCloseRequested = false;
-		TFunction<void(void)>							OnMaxReconnectAttemptsExceeded;
 
 	public:
 		// Begin EpicRtcRefCountInterface
diff --git a/Engine/Plugins/Media/PixelStreaming2/Source/PixelStreaming2RTC/Internal/EpicRtcWebsocketFactory.h b/Engine/Plugins/Media/PixelStreaming2/Source/PixelStreaming2RTC/Internal/EpicRtcWebsocketFactory.h
--- a/Engine/Plugins/Media/PixelStreaming2/Source/PixelStreaming2RTC/Internal/EpicRtcWebsocketFactory.h
+++ b/Engine/Plugins/Media/PixelStreaming2/Source/PixelStreaming2RTC/Internal/EpicRtcWebsocketFactory.h
@@ -13,11 +13,11 @@ namespace UE::PixelStreaming2
 	class PIXELSTREAMING2RTC_API FEpicRtcWebsocketFactory : public EpicRtcWebsocketFactoryInterface
 	{
 	public:
-		FEpicRtcWebsocketFactory(bool bInSendKeepAlive = true, TFunction<void(void)> OnMaxReconnectAttemptsExceeded = []() {})
+		FEpicRtcWebsocketFactory(bool bInSendKeepAlive = true)
 			: bSendKeepAlive(bInSendKeepAlive)
-			, OnMaxReconnectAttemptsExceeded(OnMaxReconnectAttemptsExceeded)
 		{
 		}
+
 		virtual ~FEpicRtcWebsocketFactory() = default;
 
 		virtual EpicRtcErrorCode CreateWebsocket(EpicRtcWebsocketInterface** OutWebsocket) override;
@@ -29,8 +29,6 @@ namespace UE::PixelStreaming2
 
 	private:
 		bool bSendKeepAlive;
-
-		TFunction<void(void)> OnMaxReconnectAttemptsExceeded;
 	};
 
 } // namespace UE::PixelStreaming2
diff --git a/Engine/Plugins/Media/PixelStreaming2/Source/PixelStreaming2RTC/Private/EpicRtcStreamer.cpp b/Engine/Plugins/Media/PixelStreaming2/Source/PixelStreaming2RTC/Private/EpicRtcStreamer.cpp
--- a/Engine/Plugins/Media/PixelStreaming2/Source/PixelStreaming2RTC/Private/EpicRtcStreamer.cpp
+++ b/Engine/Plugins/Media/PixelStreaming2/Source/PixelStreaming2RTC/Private/EpicRtcStreamer.cpp
@@ -1182,22 +1182,23 @@ namespace UE::PixelStreaming2
 			{
 				UE_LOGFMT(LogPixelStreaming2RTC, VeryVerbose, "FEpicRtcStreamer::OnSessionStateUpdate State=Disconnected");
 
+				bool bBroadcastDelegates = StreamState != EStreamState::Connecting;
+
 				// If the session unexpectedly disconnects (ie signalling server goes away), we should try and reconnect
-				if (StreamState == EStreamState::Connected)
+				if (StreamState == EStreamState::Connected || StreamState == EStreamState::Connecting)
 				{
-					// Call stop streaming first. This will remove the room
+					// Call stop streaming first
 					StopStreaming();
-					ReconnectTimer->Start(AsShared());
-				}
-				DoOnGameThread([StreamerId = this->StreamerId]() {
-					if (UPixelStreaming2Delegates* Delegates = UPixelStreaming2Delegates::Get())
+
+					if (TRefCountPtr<EpicRtcSessionInterface> PinnedSession = EpicRtcSession; EpicRtcRoom && PinnedSession)
 					{
-						Delegates->OnDisconnectedFromSignallingServer.Broadcast(StreamerId);
-						Delegates->OnDisconnectedFromSignallingServerNative.Broadcast(StreamerId);
+						FUtf8String Utf8StreamerId(StreamerId);
+						PinnedSession->RemoveRoom(ToEpicRtcStringView(Utf8StreamerId));
+						EpicRtcRoom = nullptr;
 					}
-				});
 
-				OnStreamingStopped().Broadcast(this);
+					ReconnectTimer->Start(AsShared());
+				}
 
 				// The session has been disconnected (either through a call to StopStreaming or by an error)
 				// so remove it
@@ -1208,6 +1209,19 @@ namespace UE::PixelStreaming2
 					EpicRtcSession = nullptr;
 				}
 
+				if (bBroadcastDelegates)
+				{
+					DoOnGameThread([StreamerId = this->StreamerId]() {
+						if (UPixelStreaming2Delegates* Delegates = UPixelStreaming2Delegates::Get())
+						{
+							Delegates->OnDisconnectedFromSignallingServer.Broadcast(StreamerId);
+							Delegates->OnDisconnectedFromSignallingServerNative.Broadcast(StreamerId);
+						}
+					});
+
+					OnStreamingStopped().Broadcast(this);
+				}
+
 				// We are fully disconnect at this point so we update the state so we can StartStreaming again
 				StreamState = EStreamState::Disconnected;
 
@@ -1263,6 +1277,8 @@ namespace UE::PixelStreaming2
 				UE_LOGFMT(LogPixelStreaming2RTC, VeryVerbose, "FEpicRtcStreamer::OnRoomStateUpdate State=Joined");
 
 				StreamState = EStreamState::Connected;
+				ReconnectTimer->Reset();
+
 				OnStreamingStarted().Broadcast(this);
 				break;
 			}
diff --git a/Engine/Plugins/Media/PixelStreaming2/Source/PixelStreaming2RTC/Private/EpicRtcWebsocket.cpp b/Engine/Plugins/Media/PixelStreaming2/Source/PixelStreaming2RTC/Private/EpicRtcWebsocket.cpp
--- a/Engine/Plugins/Media/PixelStreaming2/Source/PixelStreaming2RTC/Private/EpicRtcWebsocket.cpp
+++ b/Engine/Plugins/Media/PixelStreaming2/Source/PixelStreaming2RTC/Private/EpicRtcWebsocket.cpp
@@ -16,10 +16,9 @@ namespace UE::PixelStreaming2
 	DECLARE_LOG_CATEGORY_EXTERN(LogEpicRtcWebsocket, Log, All);
 	DEFINE_LOG_CATEGORY(LogEpicRtcWebsocket);
 
-	FEpicRtcWebsocket::FEpicRtcWebsocket(bool bKeepAlive, TSharedPtr<IWebSocket> InWebSocket, TFunction<void(void)> InOnMaxReconnectAttemptsExceeded)
+	FEpicRtcWebsocket::FEpicRtcWebsocket(bool bKeepAlive, TSharedPtr<IWebSocket> InWebSocket)
 		: WebSocket(InWebSocket)
 		, bSendKeepAlive(bKeepAlive)
-		, OnMaxReconnectAttemptsExceeded(InOnMaxReconnectAttemptsExceeded)
 	{
 	}
 
@@ -121,41 +120,19 @@ namespace UE::PixelStreaming2
 	{
 		UE_LOG(LogEpicRtcWebsocket, Log, TEXT("Websocket connection made to: %s"), *Url);
 		bCloseRequested = false;
-		bReconnectOnError = false;
-		NumReconnectAttempts = 0;
 		LastKeepAliveCycles = FPlatformTime::Cycles64();
 		Observer->OnOpen();
 	}
 
 	void FEpicRtcWebsocket::OnConnectionError(const FString& Error)
 	{
-		if (!WebSocket->IsConnected() && UPixelStreaming2PluginSettings::CVarSignalingReconnectInterval.GetValueOnAnyThread() > 0.0f)
-		{
-			// Reconnecting case where we had not connected yet and got an error while connecting (e.g. server not up)
-			bReconnectOnError = true;
-
-			LastError = Error;
-
-			// To reconnect we must close the existing WS (amusingly this does no trigger WS `OnClosed`)
-			WebSocket->Close();
-
-			// Note: By NOT issuing `OnClosed` here we keep the EpicRtcSession in
-			// a `pending` state while we attempt to reconnect (which seems appropriate). When a connection is made it can
-			// graduate to the `connected` through `OnOpen` being fired above.
-		}
-		else
-		{
-			// In this case with we were already connected and got an error OR we have disabled reconnection
-			UE_LOG(LogEpicRtcWebsocket, Log, TEXT("Failed to connect to %s - signalling server may not be up yet. Message: \"%s\""), *Url, *Error);
-
-			// Note: Only issue `OnClosed` if we are not attempting to reconnect.
-			Observer->OnClosed();
-		}
+		// In this case with we were already connected and got an error OR we have disabled reconnection
+		UE_LOG(LogEpicRtcWebsocket, Log, TEXT("Failed to connect to %s - signalling server may not be up yet. Message: \"%s\""), *Url, *Error);
+		Observer->OnClosed();
 	}
 
 	void FEpicRtcWebsocket::OnClosed(int32 StatusCode, const FString& Reason, bool bWasClean)
 	{
-		bReconnectOnError = false;
 		UE_LOG(LogEpicRtcWebsocket, Log, TEXT("Closed connection to %s - \n\tstatus %d\n\treason: %s\n\twas clean: %s"), *Url, StatusCode, *Reason, bWasClean ? TEXT("true") : TEXT("false"));
 		Observer->OnClosed();
 	}
@@ -209,8 +186,6 @@ namespace UE::PixelStreaming2
 			return;
 		}
 
-		Reconnect();
-
 		if (bSendKeepAlive)
 		{
 			KeepAlive();
@@ -251,55 +226,4 @@ namespace UE::PixelStreaming2
 			LastKeepAliveCycles = CyclesNow;
 		}
 	}
-
-	void FEpicRtcWebsocket::Reconnect()
-	{
-		if (!bReconnectOnError)
-		{
-			return;
-		}
-
-		if (!WebSocket)
-		{
-			return;
-		}
-
-		if (WebSocket->IsConnected())
-		{
-			return;
-		}
-
-		float ReconnectIntervalSeconds = UPixelStreaming2PluginSettings::CVarSignalingReconnectInterval.GetValueOnAnyThread();
-
-		if (ReconnectIntervalSeconds <= 0.0f)
-		{
-			return;
-		}
-
-		uint64 CyclesNow = FPlatformTime::Cycles64();
-		uint64 DeltaCycles = CyclesNow - LastReconnectCycles;
-		float  DeltaSeconds = FPlatformTime::ToSeconds(DeltaCycles);
-
-		// If enough time has elapsed, try a reconnect
-		if (DeltaSeconds >= ReconnectIntervalSeconds)
-		{
-			// Check if the next attempt to reconnect will exceed the maximum number of attempts
-			if (UPixelStreaming2PluginSettings::CVarSignalingMaxReconnectAttempts.GetValueOnAnyThread() >= 0 && (NumReconnectAttempts + 1) > UPixelStreaming2PluginSettings::CVarSignalingMaxReconnectAttempts.GetValueOnAnyThread())
-			{
-				// Maxmimum exceeded so don't attempt it and instead stop the timer
-				UE_LOGFMT(LogEpicRtcWebsocket, Warning, "Maximum number of reconnect attempts ({0}) exceeded!", UPixelStreaming2PluginSettings::CVarSignalingMaxReconnectAttempts.GetValueOnAnyThread());
-				bReconnectOnError = false;
-				OnMaxReconnectAttemptsExceeded();
-				return;
-			}
-
-			NumReconnectAttempts++;
-			UE_LOGFMT(LogEpicRtcWebsocket, Log, "Failed to connect to {0}. (\"{1}\") - signalling server may not be up yet. Reconnecting... Attempt: {2}", Url, LastError, NumReconnectAttempts);
-			WebSocket->Connect();
-			LastReconnectCycles = CyclesNow;
-			// Do not try to reconnect until we hear the next error
-			bReconnectOnError = false;
-		}
-	}
-
 } // namespace UE::PixelStreaming2
diff --git a/Engine/Plugins/Media/PixelStreaming2/Source/PixelStreaming2RTC/Private/EpicRtcWebsocketFactory.cpp b/Engine/Plugins/Media/PixelStreaming2/Source/PixelStreaming2RTC/Private/EpicRtcWebsocketFactory.cpp
--- a/Engine/Plugins/Media/PixelStreaming2/Source/PixelStreaming2RTC/Private/EpicRtcWebsocketFactory.cpp
+++ b/Engine/Plugins/Media/PixelStreaming2/Source/PixelStreaming2RTC/Private/EpicRtcWebsocketFactory.cpp
@@ -7,7 +7,7 @@ namespace UE::PixelStreaming2
 
 	EpicRtcErrorCode FEpicRtcWebsocketFactory::CreateWebsocket(EpicRtcWebsocketInterface** OutWebsocket)
 	{
-		EpicRtcWebsocketInterface* Websocket = new FEpicRtcWebsocket(bSendKeepAlive, nullptr, OnMaxReconnectAttemptsExceeded);
+		EpicRtcWebsocketInterface* Websocket = new FEpicRtcWebsocket(bSendKeepAlive, nullptr);
 		// Because the ptr was created with new, we need to call AddRef ourself (ms spec compliant)
 		Websocket->AddRef();
 
