[PixelStreaming2] Fix: Audio/video component not listening to DefaultStreamer by default when doing PIE

diff --git a/Engine/Plugins/Media/PixelStreaming2/Source/PixelStreaming2/Private/Blueprints/PixelStreaming2AudioComponent.cpp b/Engine/Plugins/Media/PixelStreaming2/Source/PixelStreaming2/Private/Blueprints/PixelStreaming2AudioComponent.cpp
--- a/Engine/Plugins/Media/PixelStreaming2/Source/PixelStreaming2/Private/Blueprints/PixelStreaming2AudioComponent.cpp
+++ b/Engine/Plugins/Media/PixelStreaming2/Source/PixelStreaming2/Private/Blueprints/PixelStreaming2AudioComponent.cpp
@@ -78,14 +78,16 @@ bool UPixelStreaming2AudioComponent::StreamerListenTo(FString StreamerId, FStrin
 
 	if (StreamerId == FString())
 	{
+		FString DefaultStreamerId = PixelStreaming2Module.GetDefaultStreamerID();
+
 		TArray<FString> StreamerIds = PixelStreaming2Module.GetStreamerIds();
-		if (StreamerIds.Num() > 0)
+		if (StreamerIds.Contains(DefaultStreamerId))
 		{
-			StreamerToHear = StreamerIds[0];
+			StreamerToHear = DefaultStreamerId;
 		}
-		else
+		else if (StreamerIds.Num() > 0)
 		{
-			StreamerToHear = PixelStreaming2Module.GetDefaultStreamerID();
+			StreamerToHear = StreamerIds[0];
 		}
 	}
 	else
diff --git a/Engine/Plugins/Media/PixelStreaming2/Source/PixelStreaming2/Private/Blueprints/PixelStreaming2VideoComponent.cpp b/Engine/Plugins/Media/PixelStreaming2/Source/PixelStreaming2/Private/Blueprints/PixelStreaming2VideoComponent.cpp
--- a/Engine/Plugins/Media/PixelStreaming2/Source/PixelStreaming2/Private/Blueprints/PixelStreaming2VideoComponent.cpp
+++ b/Engine/Plugins/Media/PixelStreaming2/Source/PixelStreaming2/Private/Blueprints/PixelStreaming2VideoComponent.cpp
@@ -56,14 +56,16 @@ bool UPixelStreaming2VideoComponent::StreamerWatch(FString StreamerId, FString P
 
 	if (StreamerId == FString())
 	{
+		FString DefaultStreamerId = PixelStreaming2Module.GetDefaultStreamerID();
+
 		TArray<FString> StreamerIds = PixelStreaming2Module.GetStreamerIds();
-		if (StreamerIds.Num() > 0)
+		if (StreamerIds.Contains(DefaultStreamerId))
 		{
-			StreamerToWatch = StreamerIds[0];
+			StreamerToWatch = DefaultStreamerId;
 		}
-		else
+		else if (StreamerIds.Num() > 0)
 		{
-			StreamerToWatch = PixelStreaming2Module.GetDefaultStreamerID();
+			StreamerToWatch = StreamerIds[0];
 		}
 	}
 	else
