[PixelStreaming2] Fix: Reduce PS2 latency during frame capture

diff --git a/Engine/Plugins/Media/PixelStreaming2/Source/PixelStreaming2RTC/Private/EpicRtcStreamer.cpp b/Engine/Plugins/Media/PixelStreaming2/Source/PixelStreaming2RTC/Private/EpicRtcStreamer.cpp
--- a/Engine/Plugins/Media/PixelStreaming2/Source/PixelStreaming2RTC/Private/EpicRtcStreamer.cpp
+++ b/Engine/Plugins/Media/PixelStreaming2/Source/PixelStreaming2RTC/Private/EpicRtcStreamer.cpp
@@ -30,7 +30,6 @@
 #include "UtilsString.h"
 #include "UtilsCodecs.h"
 #include "VideoProducerBackBuffer.h"
-#include "VideoProducerMediaCapture.h"
 #include "RTCInputHandler.h"
 #include "PixelStreaming2Common.h"
 #include "EpicRtcAudioCapturer.h"
@@ -172,11 +171,6 @@ namespace UE::PixelStreaming2
 
 		VideoCapturer->OnFrameCapturerCreated.AddSP(Streamer.ToSharedRef(), &FEpicRtcStreamer::OnFrameCapturerCreated);
 
-		if (UPixelStreaming2PluginSettings::FDelegates* Delegates = UPixelStreaming2PluginSettings::Delegates())
-		{
-			Delegates->OnUseMediaCaptureChanged.AddSP(Streamer.ToSharedRef(), &FEpicRtcStreamer::OnUseMediaCaptureChanged);
-		}
-
 		if (UPixelStreaming2Delegates* Delegates = UPixelStreaming2Delegates::Get())
 		{
 			Delegates->OnStatChangedNative.AddSP(Streamer.ToSharedRef(), &FEpicRtcStreamer::ConsumeStats);
@@ -351,32 +345,6 @@ namespace UE::PixelStreaming2
 		SendAllPlayersMessage(MessageName, Descriptor);
 	}
 
-	void FEpicRtcStreamer::OnUseMediaCaptureChanged(IConsoleVariable* CVar)
-	{
-		TSharedPtr<IPixelStreaming2VideoProducer> VideoProducer = VideoCapturer->GetVideoProducer();
-		if (!VideoProducer)
-		{
-			return;
-		}
-
-		FString Identifier = VideoProducer->ToString();
-		// If the identifier is not one of these two types, it means it was user created so we don't want to
-		// override it
-		if (Identifier != VideoProducerIdentifiers::FVideoProducerMediaCapture && Identifier != VideoProducerIdentifiers::FVideoProducerBackBuffer)
-		{
-			return;
-		}
-
-		if (CVar->GetBool())
-		{
-			SetVideoProducer(FVideoProducerMediaCapture::CreateActiveViewportCapture());
-		}
-		else
-		{
-			SetVideoProducer(FVideoProducerBackBuffer::Create());
-		}
-	}
-
 	EpicRtcBitrate FEpicRtcStreamer::GetBitrates() const
 	{
 		const int MinBitrate = UPixelStreaming2PluginSettings::CVarWebRTCMinBitrate.GetValueOnAnyThread();
@@ -2007,14 +1975,7 @@ namespace UE::PixelStreaming2
 		// the video producer in order to not tax their GPU unnecessarily
 		if (!UPixelStreaming2PluginSettings::CVarConnectionURL.GetValueOnAnyThread().IsEmpty())
 		{
-			if (UPixelStreaming2PluginSettings::CVarUseMediaCapture.GetValueOnAnyThread())
-			{
-				NewStreamer->SetVideoProducer(FVideoProducerMediaCapture::CreateActiveViewportCapture());
-			}
-			else
-			{
-				NewStreamer->SetVideoProducer(FVideoProducerBackBuffer::Create());
-			}
+			NewStreamer->SetVideoProducer(FVideoProducerBackBuffer::Create());
 		}
 
 		return NewStreamer;
diff --git a/Engine/Plugins/Media/PixelStreaming2/Source/PixelStreaming2RTC/Private/EpicRtcStreamer.h b/Engine/Plugins/Media/PixelStreaming2/Source/PixelStreaming2RTC/Private/EpicRtcStreamer.h
--- a/Engine/Plugins/Media/PixelStreaming2/Source/PixelStreaming2RTC/Private/EpicRtcStreamer.h
+++ b/Engine/Plugins/Media/PixelStreaming2/Source/PixelStreaming2RTC/Private/EpicRtcStreamer.h
@@ -122,7 +122,6 @@ namespace UE::PixelStreaming2
 		void		   OnFrameCapturerCreated();
 		void		   OnUIInteraction(FMemoryReader Ar);
 		void		   OnSendMessage(FString MessageName, FMemoryReader Ar);
-		void		   OnUseMediaCaptureChanged(IConsoleVariable* CVar);
 		EpicRtcBitrate GetBitrates() const;
 
 		template <typename T>
