[PixelStreaming2] Fix: "Stop Streaming" toolbar button fails when no active connection

diff --git a/Engine/Plugins/Media/PixelStreaming2/Source/PixelStreaming2RTC/Private/EpicRtcStreamer.cpp b/Engine/Plugins/Media/PixelStreaming2/Source/PixelStreaming2RTC/Private/EpicRtcStreamer.cpp
--- a/Engine/Plugins/Media/PixelStreaming2/Source/PixelStreaming2RTC/Private/EpicRtcStreamer.cpp
+++ b/Engine/Plugins/Media/PixelStreaming2/Source/PixelStreaming2RTC/Private/EpicRtcStreamer.cpp
@@ -277,6 +277,8 @@ namespace UE::PixelStreaming2
 	{
 		if (StreamState == EStreamState::Disconnected || StreamState == EStreamState::Disconnecting)
 		{
+			// We've received another call to StopStreaming while we're mid disconnection. Stop the reconnect timer if it was active
+			ReconnectTimer->Stop();
 			return;
 		}
 
@@ -1150,7 +1152,8 @@ namespace UE::PixelStreaming2
 			{
 				UE_LOGFMT(LogPixelStreaming2RTC, VeryVerbose, "FEpicRtcStreamer::OnSessionStateUpdate State=Disconnected");
 
-				bool bBroadcastDelegates = StreamState != EStreamState::Connecting;
+				// If we were successfully connected previously then we should broadcast the delegates about disconnection
+				bool bBroadcastDelegates = StreamState == EStreamState::Connected;
 
 				// If the session unexpectedly disconnects (ie signalling server goes away), we should try and reconnect
 				if (StreamState == EStreamState::Connected || StreamState == EStreamState::Connecting)
diff --git a/Engine/Plugins/Media/PixelStreaming2/Source/PixelStreaming2RTC/Private/EpicRtcStreamer.h b/Engine/Plugins/Media/PixelStreaming2/Source/PixelStreaming2RTC/Private/EpicRtcStreamer.h
--- a/Engine/Plugins/Media/PixelStreaming2/Source/PixelStreaming2RTC/Private/EpicRtcStreamer.h
+++ b/Engine/Plugins/Media/PixelStreaming2/Source/PixelStreaming2RTC/Private/EpicRtcStreamer.h
@@ -61,7 +61,7 @@ namespace UE::PixelStreaming2
 		virtual bool	IsConnected() override { return StreamState == EStreamState::Connected; }
 		virtual void	StartStreaming() override;
 		virtual void	StopStreaming() override;
-		virtual bool	IsStreaming() const override { return StreamState != EStreamState::Disconnected; }
+		virtual bool	IsStreaming() const override { return StreamState == EStreamState::Connected || StreamState == EStreamState::Connecting; }
 
 		virtual FPreConnectionEvent&	OnPreConnection() override;
 		virtual FStreamingStartedEvent& OnStreamingStarted() override;
