[PS2] Only tick conference when streaming.

diff --git a/Engine/Plugins/Media/PixelStreaming2/Source/PixelStreaming2/Internal/AudioCapturer.h b/Engine/Plugins/Media/PixelStreaming2/Source/PixelStreaming2/Internal/AudioCapturer.h
--- a/Engine/Plugins/Media/PixelStreaming2/Source/PixelStreaming2/Internal/AudioCapturer.h
+++ b/Engine/Plugins/Media/PixelStreaming2/Source/PixelStreaming2/Internal/AudioCapturer.h
@@ -94,8 +94,8 @@ namespace UE::PixelStreaming2
 		void OnAudio(const float* AudioData, int32 NumSamples, int32 NumChannels, int32 SampleRate);
 
 	protected:
-		TSharedPtr<FAudioPatchMixer>  Mixer;
-		TUniqueTaskPtr<FMixAudioTask> MixerTask;
+		TSharedPtr<FAudioPatchMixer> Mixer;
+		TSharedPtr<FMixAudioTask>    MixerTask;
 
 		TMap<Audio::FDeviceId, TSharedPtr<FAudioProducer>> AudioProducers;
 
diff --git a/Engine/Plugins/Media/PixelStreaming2/Source/PixelStreaming2/Internal/TickableTask.h b/Engine/Plugins/Media/PixelStreaming2/Source/PixelStreaming2/Internal/TickableTask.h
--- a/Engine/Plugins/Media/PixelStreaming2/Source/PixelStreaming2/Internal/TickableTask.h
+++ b/Engine/Plugins/Media/PixelStreaming2/Source/PixelStreaming2/Internal/TickableTask.h
@@ -3,6 +3,7 @@
 #pragma once
 
 #include "Templates/UniquePtr.h"
+#include "Templates/SharedPointer.h"
 
 namespace UE::PixelStreaming2
 {
@@ -10,9 +11,10 @@ namespace UE::PixelStreaming2
 	using TUniqueTaskPtr = TUniquePtr<T, struct FPixelStreamingTaskDeleter>;
 
 	/**
-	 * Base class for a tickable task. Inheriting from this class ensures that your task can be ticked by the PixelStreaming thread
+	 * Base class for a tickable task. Inheriting from this class ensures that your task can be ticked by the PixelStreaming thread.
+	 * Note it is important to destroy the task when it is no longer in use to save the tick thread.
 	 */
-	class PIXELSTREAMING2_API FPixelStreamingTickableTask
+	class PIXELSTREAMING2_API FPixelStreamingTickableTask : public TSharedFromThis<FPixelStreamingTickableTask>
 	{
 	public:
 		/**
@@ -21,7 +23,7 @@ namespace UE::PixelStreaming2
 		 * the PixelStreaming thread.
 		 */
 		template <typename T, typename... Args>
-		static TUniqueTaskPtr<T> Create(Args&&...);
+		static TSharedPtr<T> Create(Args&&... InArgs);
 
 	public:
 		virtual ~FPixelStreamingTickableTask() = default;
@@ -57,11 +59,11 @@ namespace UE::PixelStreaming2
 	};
 
 	template <typename T, typename... Args>
-	TUniqueTaskPtr<T> FPixelStreamingTickableTask::Create(Args&&... InArgs)
+	TSharedPtr<T> FPixelStreamingTickableTask::Create(Args&&... InArgs)
 	{
 		static_assert(std::is_base_of_v<FPixelStreamingTickableTask, T>);
 
-		TUniqueTaskPtr<T> Task(new T(Forward<Args>(InArgs)...), FPixelStreamingTaskDeleter());
+		TSharedPtr<T> Task(new T(Forward<Args>(InArgs)...), FPixelStreamingTaskDeleter());
 		Task->Register();
 
 		return Task;
diff --git a/Engine/Plugins/Media/PixelStreaming2/Source/PixelStreaming2/Private/Thread.cpp b/Engine/Plugins/Media/PixelStreaming2/Source/PixelStreaming2/Private/Thread.cpp
--- a/Engine/Plugins/Media/PixelStreaming2/Source/PixelStreaming2/Private/Thread.cpp
+++ b/Engine/Plugins/Media/PixelStreaming2/Source/PixelStreaming2/Private/Thread.cpp
@@ -74,9 +74,10 @@ namespace UE::PixelStreaming2
 			for (auto& Task : Tasks)
 			{
 				// A task may be nulled out due to deletion during our loop. Check for safety
-				if (Task)
+				TSharedPtr<FPixelStreamingTickableTask> PinnedTask = Task.Pin();
+				if (PinnedTask)
 				{
-					Task->Tick(DeltaMs);
+					PinnedTask->Tick(DeltaMs);
 				}
 			}
 
@@ -103,7 +104,7 @@ namespace UE::PixelStreaming2
 		{
 			if (bNeedsCleanup)
 			{
-				Tasks.RemoveAll([](const FPixelStreamingTickableTask* Entry) { return Entry == nullptr; });
+				Tasks.RemoveAll([](const TWeakPtr<FPixelStreamingTickableTask> Entry) { return Entry == nullptr || !Entry.IsValid(); });
 				bNeedsCleanup = false;
 			}
 
@@ -128,7 +129,7 @@ namespace UE::PixelStreaming2
 		virtual ~FPixelStreamingRunnable() = default;
 
 	private:
-		void AddTask(FPixelStreamingTickableTask* Task)
+		void AddTask(TWeakPtr<FPixelStreamingTickableTask> Task)
 		{
 			FScopeLock NewTaskLock(&NewTasksMutex);
 			NewTasks.Add(Task);
@@ -149,22 +150,30 @@ namespace UE::PixelStreaming2
 			FScopeLock NewTaskLock(&NewTasksMutex);
 
 			// Remove from pending list if it hasn't been registered
-			NewTasks.Remove(Task);
+			NewTasks.RemoveAll([Task](const TWeakPtr<FPixelStreamingTickableTask> Entry) { 
+				const TSharedPtr<FPixelStreamingTickableTask> Pin = Entry.Pin();
+				return Entry == nullptr || !Pin || Pin.Get() == Task; 
+				});
 
-			const int32 Pos = Tasks.IndexOfByKey(Task);
-			if (Pos != INDEX_NONE)
+			// During ticking it is not safe to modify the set so null and mark for later
+			if (bIsTicking)
 			{
-				if (bIsTicking)
+				bNeedsCleanup = true;
+				for (TWeakPtr<FPixelStreamingTickableTask>& LoopTask : Tasks)
 				{
-					// During ticking it is not safe to modify the set so null and mark for later
-					Tasks[Pos] = nullptr;
-					bNeedsCleanup = true;
-				}
-				else
-				{
-					Tasks.RemoveAt(Pos);
+					if (TSharedPtr<FPixelStreamingTickableTask> Pin = LoopTask.Pin(); Pin && Pin.Get() == Task)
+					{
+						LoopTask = nullptr;
+					}
 				}
 			}
+			else
+			{
+				Tasks.RemoveAll([Task](const TWeakPtr<FPixelStreamingTickableTask> Entry) {
+					TSharedPtr<FPixelStreamingTickableTask> Pin = Entry.Pin();
+					return Entry == nullptr || !Pin || Pin.Get() == Task;
+				});
+			}
 		}
 
 	private:
@@ -172,12 +181,12 @@ namespace UE::PixelStreaming2
 		friend FPixelStreamingTickableTask;
 
 		// New tasks that have not yet been added to the Tasks list*/
-		TArray<FPixelStreamingTickableTask*> NewTasks;
+		TArray<TWeakPtr<FPixelStreamingTickableTask>> NewTasks;
 		// Lock for modifying new list */
 		FCriticalSection NewTasksMutex;
 
 		// Tasks to execute every tick
-		TArray<FPixelStreamingTickableTask*> Tasks;
+		TArray<TWeakPtr<FPixelStreamingTickableTask>> Tasks;
 		// This critical section should be locked during entire tick process
 		FCriticalSection TasksMutex;
 
@@ -224,7 +233,7 @@ namespace UE::PixelStreaming2
 		TSharedPtr<FPixelStreamingRunnable> Runnable = PixelStreamingRunnable.Pin();
 		if (Runnable)
 		{
-			Runnable->AddTask(this);
+			Runnable->AddTask(AsWeak());
 		}
 	}
 
diff --git a/Engine/Plugins/Media/PixelStreaming2/Source/PixelStreaming2RTC/Internal/SharedTickableTasks.h b/Engine/Plugins/Media/PixelStreaming2/Source/PixelStreaming2RTC/Internal/SharedTickableTasks.h
new file mode 100644
--- /dev/null
+++ b/Engine/Plugins/Media/PixelStreaming2/Source/PixelStreaming2RTC/Internal/SharedTickableTasks.h
@@ -0,0 +1,33 @@
+// Copyright Epic Games, Inc. All Rights Reserved.
+
+#pragma once
+
+#include "EpicRtcAudioCapturer.h"
+#include "EpicRtcConferenceUtils.h"
+#include "TickableTask.h"
+
+namespace UE::PixelStreaming2
+{
+	class FSharedTickableTasks
+	{
+	public:
+		FSharedTickableTasks(
+			TSharedPtr<FEpicRtcTickConferenceTask>&& InTickConferenceTask)
+			: TickConferenceTask(MoveTemp(InTickConferenceTask)){};
+		
+
+		TSharedPtr<FEpicRtcAudioCapturer> GetAudioCapturer()
+		{
+			if (!AudioCapturer)
+			{
+				AudioCapturer = FEpicRtcAudioCapturer::Create();
+			}
+			return AudioCapturer;
+		}
+
+	private:
+		
+		TSharedPtr<FEpicRtcTickConferenceTask> TickConferenceTask;
+		TSharedPtr<FEpicRtcAudioCapturer>	   AudioCapturer;
+	};
+} // namespace UE::PixelStreaming2
diff --git a/Engine/Plugins/Media/PixelStreaming2/Source/PixelStreaming2RTC/Private/EpicRtcStreamer.cpp b/Engine/Plugins/Media/PixelStreaming2/Source/PixelStreaming2RTC/Private/EpicRtcStreamer.cpp
--- a/Engine/Plugins/Media/PixelStreaming2/Source/PixelStreaming2RTC/Private/EpicRtcStreamer.cpp
+++ b/Engine/Plugins/Media/PixelStreaming2/Source/PixelStreaming2RTC/Private/EpicRtcStreamer.cpp
@@ -25,6 +25,7 @@
 #include "EpicRtcTrack.h"
 #include "EpicRtcDataTrack.h"
 #include "IPixelStreaming2InputModule.h"
+#include "IPixelStreaming2Module.h"
 #include "CoreGlobals.h"
 #include "UtilsAsync.h"
 #include "UtilsString.h"
@@ -243,6 +244,8 @@ namespace UE::PixelStreaming2
 
 		VideoCapturer->ResetFrameCapturer();
 
+		TickableTasks = FPixelStreaming2RTCModule::GetModule()->GetSharedTickableTasks();
+
 		// Broadcast the preconnection event just before we do `TryConnect`
 		OnPreConnection().Broadcast(this);
 
@@ -293,6 +296,8 @@ namespace UE::PixelStreaming2
 		TriggerMouseLeave(StreamerId);
 
 		DeleteAllPlayerSessions();
+
+		TickableTasks.Reset();
 	}
 
 	void FEpicRtcStreamer::OnStatsReady(const FString& PlayerId, const EpicRtcConnectionStats& ConnectionStats)
@@ -1550,7 +1555,7 @@ namespace UE::PixelStreaming2
 			}
 			else
 			{
-				Participant->AudioSource = FEpicRtcAudioSource::Create(AudioTrack, FPixelStreaming2RTCModule::GetModule()->GetAudioCapturer());
+				Participant->AudioSource = FEpicRtcAudioSource::Create(AudioTrack, TickableTasks->GetAudioCapturer());
 			}
 		}
 
diff --git a/Engine/Plugins/Media/PixelStreaming2/Source/PixelStreaming2RTC/Private/EpicRtcStreamer.h b/Engine/Plugins/Media/PixelStreaming2/Source/PixelStreaming2RTC/Private/EpicRtcStreamer.h
--- a/Engine/Plugins/Media/PixelStreaming2/Source/PixelStreaming2RTC/Private/EpicRtcStreamer.h
+++ b/Engine/Plugins/Media/PixelStreaming2/Source/PixelStreaming2RTC/Private/EpicRtcStreamer.h
@@ -16,6 +16,7 @@
 #include "IPixelStreaming2Streamer.h"
 #include "Logging.h"
 #include "PlayerContext.h"
+#include "SharedTickableTasks.h"
 #include "StreamerReconnectTimer.h"
 #include "Templates/SharedPointer.h"
 #include "ThreadSafeMap.h"
@@ -201,6 +202,7 @@ namespace UE::PixelStreaming2
 		TMap<uintptr_t /* EpicRtcVideoTrackInterface* */, FString /*participant id*/> VideoTrackPlayerIdMap;
 
 		TSharedPtr<FStreamerReconnectTimer> ReconnectTimer;
+		TSharedPtr<FSharedTickableTasks>	TickableTasks;
 
 	public:
 		// Begin IPixelStreaming2SessionObserver
diff --git a/Engine/Plugins/Media/PixelStreaming2/Source/PixelStreaming2RTC/Private/PixelStreaming2RTCModule.cpp b/Engine/Plugins/Media/PixelStreaming2/Source/PixelStreaming2RTC/Private/PixelStreaming2RTCModule.cpp
--- a/Engine/Plugins/Media/PixelStreaming2/Source/PixelStreaming2RTC/Private/PixelStreaming2RTCModule.cpp
+++ b/Engine/Plugins/Media/PixelStreaming2/Source/PixelStreaming2RTC/Private/PixelStreaming2RTCModule.cpp
@@ -132,15 +132,16 @@ namespace UE::PixelStreaming2
 				});
 			});
 			Delegates->OnWebRTCDisableStatsChanged.AddLambda([this](IConsoleVariable* Var) {
-				if (EpicRtcConference)
+				TRefCountPtr<EpicRtcConferenceInterface> Conference(EpicRtcConference);
+				if (Conference)
 				{
 					if (Var->GetBool())
 					{
-						EpicRtcConference->DisableStats();
+						Conference->DisableStats();
 					}
 					else
 					{
-						EpicRtcConference->EnableStats();
+						Conference->EnableStats();
 					}
 				}
 			});
@@ -171,9 +172,7 @@ namespace UE::PixelStreaming2
 			MediaModulePtr->UnregisterPlayerFactory(*this);
 		}
 
-		AudioMixingCapturer.Reset();
-
-		TickConferenceTask.Reset();
+		TickableTasks.Reset();
 		StreamerFactory.Reset();
 
 		if (!EpicRtcPlatform)
@@ -261,14 +260,15 @@ namespace UE::PixelStreaming2
 	 * End IMediaPlayerFactory implementation
 	 */
 
-	TSharedPtr<FEpicRtcAudioCapturer> FPixelStreaming2RTCModule::GetAudioCapturer()
+	TSharedPtr<FSharedTickableTasks> FPixelStreaming2RTCModule::GetSharedTickableTasks()
 	{
-		if (!AudioMixingCapturer)
+		TSharedPtr<FSharedTickableTasks> PinnedTickableTasks = TickableTasks.Pin();
+		if (!PinnedTickableTasks)
 		{
-			AudioMixingCapturer = FEpicRtcAudioCapturer::Create();
+			PinnedTickableTasks = MakeShared<FSharedTickableTasks>(FPixelStreamingTickableTask::Create<FEpicRtcTickConferenceTask>(EpicRtcConference, TEXT("PixelStreaming2Module TickConferenceTask")));
+			TickableTasks = PinnedTickableTasks;
 		}
-
-		return AudioMixingCapturer;
+		return PinnedTickableTasks;
 	}
 
 	FString FPixelStreaming2RTCModule::GetFieldTrials()
@@ -385,8 +385,6 @@ namespace UE::PixelStreaming2
 			return false;
 		}
 
-		TickConferenceTask = FPixelStreamingTickableTask::Create<FEpicRtcTickConferenceTask>(EpicRtcConference, TEXT("PixelStreaming2Module TickConferenceTask"));
-
 		return true;
 	}
 
diff --git a/Engine/Plugins/Media/PixelStreaming2/Source/PixelStreaming2RTC/Private/PixelStreaming2RTCModule.h b/Engine/Plugins/Media/PixelStreaming2/Source/PixelStreaming2RTC/Private/PixelStreaming2RTCModule.h
--- a/Engine/Plugins/Media/PixelStreaming2/Source/PixelStreaming2RTC/Private/PixelStreaming2RTCModule.h
+++ b/Engine/Plugins/Media/PixelStreaming2/Source/PixelStreaming2RTC/Private/PixelStreaming2RTCModule.h
@@ -40,7 +40,7 @@ namespace UE::PixelStreaming2
 		virtual const TArray<FString>&	 GetSupportedPlatforms() const override;
 		// End IMediaPlayerFactory
 
-		TSharedPtr<class FEpicRtcAudioCapturer>	  GetAudioCapturer();
+		TSharedPtr<FSharedTickableTasks>		  GetSharedTickableTasks();
 		TRefCountPtr<EpicRtcConferenceInterface>& GetEpicRtcConference() { return EpicRtcConference; }
 		TRefCountPtr<FEpicRtcStatsCollector>&	  GetStatsCollector() { return StatsCollector; }
 
@@ -62,13 +62,12 @@ namespace UE::PixelStreaming2
 		FDelegateHandle LogStatsHandle;
 
 	private:
-		TSharedPtr<class FEpicRtcAudioCapturer>	 AudioMixingCapturer;
 		TRefCountPtr<EpicRtcPlatformInterface>	 EpicRtcPlatform;
 		TRefCountPtr<EpicRtcConferenceInterface> EpicRtcConference;
 		TRefCountPtr<FEpicRtcStatsCollector>	 StatsCollector;
 
-		TRefCountPtr<FEpicRtcWebsocketFactory>	   WebsocketFactory;
-		TUniqueTaskPtr<FEpicRtcTickConferenceTask> TickConferenceTask;
+		TRefCountPtr<FEpicRtcWebsocketFactory> WebsocketFactory;
+		TWeakPtr<FSharedTickableTasks>		   TickableTasks;
 
 		TArray<EpicRtcVideoEncoderInitializerInterface*> EpicRtcVideoEncoderInitializers;
 		TArray<EpicRtcVideoDecoderInitializerInterface*> EpicRtcVideoDecoderInitializers;
diff --git a/Engine/Plugins/Media/PixelStreaming2/Source/PixelStreaming2RTC/Private/Tests/MockPlayer.h b/Engine/Plugins/Media/PixelStreaming2/Source/PixelStreaming2RTC/Private/Tests/MockPlayer.h
--- a/Engine/Plugins/Media/PixelStreaming2/Source/PixelStreaming2RTC/Private/Tests/MockPlayer.h
+++ b/Engine/Plugins/Media/PixelStreaming2/Source/PixelStreaming2RTC/Private/Tests/MockPlayer.h
@@ -163,7 +163,7 @@ namespace UE::PixelStreaming2
 		TSharedPtr<FMockVideoSink>				   VideoSink;
 		TSharedPtr<FEpicRtcDataTrack>			   DataTrack;
 		TRefCountPtr<EpicRtcPlatformInterface>	   Platform;
-		TUniqueTaskPtr<FEpicRtcTickConferenceTask> TickConferenceTask;
+		TSharedPtr<FEpicRtcTickConferenceTask>     TickConferenceTask;
 		TSharedPtr<IPixelStreaming2DataProtocol>   ToStreamerProtocol;
 
 		TArray<EpicRtcVideoEncoderInitializerInterface*> EpicRtcVideoEncoderInitializers;
