Disable PixelStreaming2 on Windows Arm64

#rnx

diff --git a/Engine/Plugins/Media/PixelStreaming2/PixelStreaming2.uplugin b/Engine/Plugins/Media/PixelStreaming2/PixelStreaming2.uplugin
--- a/Engine/Plugins/Media/PixelStreaming2/PixelStreaming2.uplugin
+++ b/Engine/Plugins/Media/PixelStreaming2/PixelStreaming2.uplugin
@@ -20,48 +20,55 @@
 			"Type": "Runtime",
 			"LoadingPhase": "PreDefault",
 			"PlatformAllowList": [ "Win64", "Linux", "Mac" ],
-			"TargetDenyList": [ ]
+			"PlatformArchitectureDenyList": [ "Win64:arm64" ],
+			"TargetDenyList": []
 		},
 		{
 			"Name": "PixelStreaming2",
 			"Type": "Runtime",
 			"LoadingPhase": "PreDefault",
 			"PlatformAllowList": [ "Win64", "Linux", "Mac" ],
-			"TargetDenyList": [ ]
+			"PlatformArchitectureDenyList": [ "Win64:arm64" ],
+			"TargetDenyList": []
 		},
 		{
 			"Name": "PixelStreaming2Settings",
 			"Type": "Runtime",
 			"LoadingPhase": "PreDefault",
 			"PlatformAllowList": [ "Win64", "Linux", "Mac" ],
-			"TargetDenyList": [ ]
+			"PlatformArchitectureDenyList": [ "Win64:arm64" ],
+			"TargetDenyList": []
 		},
 		{
 			"Name": "PixelStreaming2RTC",
 			"Type": "Runtime",
 			"LoadingPhase": "PreDefault",
 			"PlatformAllowList": [ "Win64", "Linux", "Mac" ],
-			"TargetDenyList": [ ]
+			"PlatformArchitectureDenyList": [ "Win64:arm64" ],
+			"TargetDenyList": []
 		},
 		{
 			"Name": "PixelStreaming2Servers",
 			"Type": "Runtime",
 			"LoadingPhase": "PostConfigInit",
 			"PlatformAllowList": [ "Win64", "Linux", "Mac" ],
-			"TargetDenyList": [ ]
+			"PlatformArchitectureDenyList": [ "Win64:arm64" ],
+			"TargetDenyList": []
 		},
 		{
-            "Name": "PixelStreaming2Editor",
-            "Type": "Editor",
-            "LoadingPhase": "PreDefault",
-            "PlatformAllowList": [ "Win64", "Linux", "Mac" ],
-            "TargetDenyList": [ "Server" ]
+			"Name": "PixelStreaming2Editor",
+			"Type": "Editor",
+			"LoadingPhase": "PreDefault",
+			"PlatformAllowList": [ "Win64", "Linux", "Mac" ],
+			"PlatformArchitectureDenyList": [ "Win64:arm64" ],
+			"TargetDenyList": [ "Server" ]
 		},
 		{
 			"Name": "PixelStreaming2HMD",
 			"Type": "Runtime",
 			"LoadingPhase": "PreDefault",
 			"PlatformAllowList": [ "Win64", "Linux", "Mac" ],
+			"PlatformArchitectureDenyList": [ "Win64:arm64" ],
 			"TargetDenyList": [ "Server" ]
 		},
 		{
@@ -69,6 +76,7 @@
 			"Type": "Runtime",
 			"LoadingPhase": "PreDefault",
 			"PlatformAllowList": [ "Win64", "Linux", "Mac" ],
+			"PlatformArchitectureDenyList": [ "Win64:arm64" ],
 			"TargetDenyList": [ "Server" ]
 		}
 	],
