RTCP-7087 Remove unnecessary includes

diff --git a/Engine/Plugins/Media/PixelStreaming2/Source/PixelStreaming2/Internal/AudioCapturer.h b/Engine/Plugins/Media/PixelStreaming2/Source/PixelStreaming2/Internal/AudioCapturer.h
--- a/Engine/Plugins/Media/PixelStreaming2/Source/PixelStreaming2/Internal/AudioCapturer.h
+++ b/Engine/Plugins/Media/PixelStreaming2/Source/PixelStreaming2/Internal/AudioCapturer.h
@@ -4,7 +4,6 @@
 
 #include "DSP/MultithreadedPatching.h"
 #include "AudioProducer.h"
-#include "ISubmixBufferListener.h"
 #include "SampleBuffer.h"
 #include "TickableTask.h"
 
diff --git a/Engine/Plugins/Media/PixelStreaming2/Source/PixelStreaming2/Internal/PixelStreaming2Trace.h b/Engine/Plugins/Media/PixelStreaming2/Source/PixelStreaming2/Internal/PixelStreaming2Trace.h
--- a/Engine/Plugins/Media/PixelStreaming2/Source/PixelStreaming2/Internal/PixelStreaming2Trace.h
+++ b/Engine/Plugins/Media/PixelStreaming2/Source/PixelStreaming2/Internal/PixelStreaming2Trace.h
@@ -2,8 +2,6 @@
 
 #pragma once
 
-#include "CoreTypes.h"
-#include "CoreMinimal.h"
 #include "Trace/Trace.h"
 
 UE_TRACE_CHANNEL_EXTERN(PixelStreaming2Channel, PIXELSTREAMING2_API);
diff --git a/Engine/Plugins/Media/PixelStreaming2/Source/PixelStreaming2/Internal/SoundGenerator.h b/Engine/Plugins/Media/PixelStreaming2/Source/PixelStreaming2/Internal/SoundGenerator.h
--- a/Engine/Plugins/Media/PixelStreaming2/Source/PixelStreaming2/Internal/SoundGenerator.h
+++ b/Engine/Plugins/Media/PixelStreaming2/Source/PixelStreaming2/Internal/SoundGenerator.h
@@ -2,7 +2,6 @@
 
 #pragma once
 
-#include "CoreTypes.h"
 #include "Sound/SoundGenerator.h"
 
 #include <atomic>
diff --git a/Engine/Plugins/Media/PixelStreaming2/Source/PixelStreaming2/Internal/UtilsAsync.h b/Engine/Plugins/Media/PixelStreaming2/Source/PixelStreaming2/Internal/UtilsAsync.h
--- a/Engine/Plugins/Media/PixelStreaming2/Source/PixelStreaming2/Internal/UtilsAsync.h
+++ b/Engine/Plugins/Media/PixelStreaming2/Source/PixelStreaming2/Internal/UtilsAsync.h
@@ -3,7 +3,6 @@
 #pragma once
 
 #include "Async/Async.h"
-#include "CoreMinimal.h"
 
 namespace UE::PixelStreaming2
 {
diff --git a/Engine/Plugins/Media/PixelStreaming2/Source/PixelStreaming2/Internal/UtilsCore.h b/Engine/Plugins/Media/PixelStreaming2/Source/PixelStreaming2/Internal/UtilsCore.h
--- a/Engine/Plugins/Media/PixelStreaming2/Source/PixelStreaming2/Internal/UtilsCore.h
+++ b/Engine/Plugins/Media/PixelStreaming2/Source/PixelStreaming2/Internal/UtilsCore.h
@@ -2,10 +2,7 @@
 
 #pragma once
 
-#include "CoreGlobals.h"
-#include "CoreMinimal.h"
 #include "Framework/Application/SlateApplication.h"
-#include "Misc/App.h"
 #include "Misc/MessageDialog.h"
 
 #if PLATFORM_WINDOWS
diff --git a/Engine/Plugins/Media/PixelStreaming2/Source/PixelStreaming2/Internal/UtilsVideo.h b/Engine/Plugins/Media/PixelStreaming2/Source/PixelStreaming2/Internal/UtilsVideo.h
--- a/Engine/Plugins/Media/PixelStreaming2/Source/PixelStreaming2/Internal/UtilsVideo.h
+++ b/Engine/Plugins/Media/PixelStreaming2/Source/PixelStreaming2/Internal/UtilsVideo.h
@@ -2,10 +2,8 @@
 
 #pragma once
 
-#include "Containers/Array.h"
 #include "PixelStreaming2PluginSettings.h"
 #include "UtilsCore.h"
-#include "Video/VideoEncoder.h"
 
 namespace UE::PixelStreaming2
 {
diff --git a/Engine/Plugins/Media/PixelStreaming2/Source/PixelStreaming2/Internal/VideoCapturer.h b/Engine/Plugins/Media/PixelStreaming2/Source/PixelStreaming2/Internal/VideoCapturer.h
--- a/Engine/Plugins/Media/PixelStreaming2/Source/PixelStreaming2/Internal/VideoCapturer.h
+++ b/Engine/Plugins/Media/PixelStreaming2/Source/PixelStreaming2/Internal/VideoCapturer.h
@@ -2,13 +2,10 @@
 
 #pragma once
 
-#include "Delegates/Delegate.h"
 #include "IPixelCaptureCapturerSource.h"
 #include "PixelCaptureCapturerMultiFormat.h"
 #include "VideoProducer.h"
 
-#include "epic_rtc/core/video/video_buffer.h"
-
 namespace UE::PixelStreaming2
 {
 	/**
diff --git a/Engine/Plugins/Media/PixelStreaming2/Source/PixelStreaming2/Internal/VideoProducer.h b/Engine/Plugins/Media/PixelStreaming2/Source/PixelStreaming2/Internal/VideoProducer.h
--- a/Engine/Plugins/Media/PixelStreaming2/Source/PixelStreaming2/Internal/VideoProducer.h
+++ b/Engine/Plugins/Media/PixelStreaming2/Source/PixelStreaming2/Internal/VideoProducer.h
@@ -2,10 +2,8 @@
 
 #pragma once
 
-#include "Delegates/DelegateCombinations.h"
 #include "IPixelStreaming2VideoProducer.h"
 #include "MediaCapture.h"
-#include "Templates/SharedPointer.h"
 
 namespace UE::PixelStreaming2
 {
diff --git a/Engine/Plugins/Media/PixelStreaming2/Source/PixelStreaming2/Internal/VideoProducerBackBuffer.h b/Engine/Plugins/Media/PixelStreaming2/Source/PixelStreaming2/Internal/VideoProducerBackBuffer.h
--- a/Engine/Plugins/Media/PixelStreaming2/Source/PixelStreaming2/Internal/VideoProducerBackBuffer.h
+++ b/Engine/Plugins/Media/PixelStreaming2/Source/PixelStreaming2/Internal/VideoProducerBackBuffer.h
@@ -4,8 +4,6 @@
 
 #include "VideoProducer.h"
 #include "Widgets/SWindow.h"
-#include "RHI.h"
-#include "Delegates/IDelegateInstance.h"
 
 namespace UE::PixelStreaming2
 {
diff --git a/Engine/Plugins/Media/PixelStreaming2/Source/PixelStreaming2/Internal/VideoProducerMediaCapture.h b/Engine/Plugins/Media/PixelStreaming2/Source/PixelStreaming2/Internal/VideoProducerMediaCapture.h
--- a/Engine/Plugins/Media/PixelStreaming2/Source/PixelStreaming2/Internal/VideoProducerMediaCapture.h
+++ b/Engine/Plugins/Media/PixelStreaming2/Source/PixelStreaming2/Internal/VideoProducerMediaCapture.h
@@ -2,11 +2,8 @@
 
 #pragma once
 
-#include "Delegates/IDelegateInstance.h"
-#include "Misc/Optional.h"
 #include "MediaOutput.h"
 #include "PixelStreaming2MediaIOCapture.h"
-#include "VideoProducer.h"
 
 
 #include "VideoProducerMediaCapture.generated.h"
diff --git a/Engine/Plugins/Media/PixelStreaming2/Source/PixelStreaming2/Internal/VideoProducerPIEViewport.h b/Engine/Plugins/Media/PixelStreaming2/Source/PixelStreaming2/Internal/VideoProducerPIEViewport.h
--- a/Engine/Plugins/Media/PixelStreaming2/Source/PixelStreaming2/Internal/VideoProducerPIEViewport.h
+++ b/Engine/Plugins/Media/PixelStreaming2/Source/PixelStreaming2/Internal/VideoProducerPIEViewport.h
@@ -2,7 +2,6 @@
 
 #pragma once
 
-#include "Delegates/IDelegateInstance.h"
 #include "VideoProducer.h"
 
 class FViewport;
diff --git a/Engine/Plugins/Media/PixelStreaming2/Source/PixelStreaming2/Internal/VideoProducerRenderTarget.h b/Engine/Plugins/Media/PixelStreaming2/Source/PixelStreaming2/Internal/VideoProducerRenderTarget.h
--- a/Engine/Plugins/Media/PixelStreaming2/Source/PixelStreaming2/Internal/VideoProducerRenderTarget.h
+++ b/Engine/Plugins/Media/PixelStreaming2/Source/PixelStreaming2/Internal/VideoProducerRenderTarget.h
@@ -2,7 +2,6 @@
 
 #pragma once
 
-#include "Delegates/IDelegateInstance.h"
 #include "VideoProducer.h"
 
 class UTextureRenderTarget2D;
diff --git a/Engine/Plugins/Media/PixelStreaming2/Source/PixelStreaming2/Internal/VideoSink.h b/Engine/Plugins/Media/PixelStreaming2/Source/PixelStreaming2/Internal/VideoSink.h
--- a/Engine/Plugins/Media/PixelStreaming2/Source/PixelStreaming2/Internal/VideoSink.h
+++ b/Engine/Plugins/Media/PixelStreaming2/Source/PixelStreaming2/Internal/VideoSink.h
@@ -6,6 +6,7 @@
 #include "CoreTypes.h"
 #include "HAL/CriticalSection.h"
 #include "IPixelStreaming2VideoSink.h"
+#include "Rendering/SlateRenderer.h"
 
 namespace UE::PixelStreaming2
 {
diff --git a/Engine/Plugins/Media/PixelStreaming2/Source/PixelStreaming2/Internal/VideoSource.h b/Engine/Plugins/Media/PixelStreaming2/Source/PixelStreaming2/Internal/VideoSource.h
--- a/Engine/Plugins/Media/PixelStreaming2/Source/PixelStreaming2/Internal/VideoSource.h
+++ b/Engine/Plugins/Media/PixelStreaming2/Source/PixelStreaming2/Internal/VideoSource.h
@@ -2,12 +2,7 @@
 
 #pragma once
 
-#include "HAL/Event.h"
-#include "HAL/IConsoleManager.h"
-#include "HAL/Runnable.h"
-#include "HAL/RunnableThread.h"
-#include "Misc/SingleThreadRunnable.h"
-#include "Templates/SharedPointer.h"
+
 #include "VideoCapturer.h"
 
 namespace UE::PixelStreaming2
diff --git a/Engine/Plugins/Media/PixelStreaming2/Source/PixelStreaming2/Internal/VideoSourceGroup.h b/Engine/Plugins/Media/PixelStreaming2/Source/PixelStreaming2/Internal/VideoSourceGroup.h
--- a/Engine/Plugins/Media/PixelStreaming2/Source/PixelStreaming2/Internal/VideoSourceGroup.h
+++ b/Engine/Plugins/Media/PixelStreaming2/Source/PixelStreaming2/Internal/VideoSourceGroup.h
@@ -2,13 +2,7 @@
 
 #pragma once
 
-#include "HAL/Event.h"
-#include "HAL/IConsoleManager.h"
-#include "HAL/Runnable.h"
-#include "HAL/RunnableThread.h"
-#include "Misc/SingleThreadRunnable.h"
 #include "VideoCapturer.h"
-#include "Templates/SharedPointer.h"
 
 namespace UE::PixelStreaming2
 {
diff --git a/Engine/Plugins/Media/PixelStreaming2/Source/PixelStreaming2/Private/AudioSink.cpp b/Engine/Plugins/Media/PixelStreaming2/Source/PixelStreaming2/Private/AudioSink.cpp
--- a/Engine/Plugins/Media/PixelStreaming2/Source/PixelStreaming2/Private/AudioSink.cpp
+++ b/Engine/Plugins/Media/PixelStreaming2/Source/PixelStreaming2/Private/AudioSink.cpp
@@ -2,7 +2,6 @@
 
 #include "AudioSink.h"
 
-#include "IPixelStreaming2AudioConsumer.h"
 #include "PixelStreaming2Trace.h"
 
 namespace UE::PixelStreaming2
diff --git a/Engine/Plugins/Media/PixelStreaming2/Source/PixelStreaming2/Private/Blueprints/PixelStreaming2InputComponent.cpp b/Engine/Plugins/Media/PixelStreaming2/Source/PixelStreaming2/Private/Blueprints/PixelStreaming2InputComponent.cpp
--- a/Engine/Plugins/Media/PixelStreaming2/Source/PixelStreaming2/Private/Blueprints/PixelStreaming2InputComponent.cpp
+++ b/Engine/Plugins/Media/PixelStreaming2/Source/PixelStreaming2/Private/Blueprints/PixelStreaming2InputComponent.cpp
@@ -9,6 +9,7 @@
 #include "Logging.h"
 #include "PixelStreaming2Common.h"
 #include "PixelStreaming2Utils.h"
+#include "PixelStreaming2InputEnums.h"
 
 UPixelStreaming2Input::UPixelStreaming2Input(const FObjectInitializer& ObjectInitializer)
 	: Super(ObjectInitializer)
diff --git a/Engine/Plugins/Media/PixelStreaming2/Source/PixelStreaming2/Private/Blueprints/PixelStreaming2StreamerComponent.cpp b/Engine/Plugins/Media/PixelStreaming2/Source/PixelStreaming2/Private/Blueprints/PixelStreaming2StreamerComponent.cpp
--- a/Engine/Plugins/Media/PixelStreaming2/Source/PixelStreaming2/Private/Blueprints/PixelStreaming2StreamerComponent.cpp
+++ b/Engine/Plugins/Media/PixelStreaming2/Source/PixelStreaming2/Private/Blueprints/PixelStreaming2StreamerComponent.cpp
@@ -3,6 +3,7 @@
 #include "PixelStreaming2StreamerComponent.h"
 
 #include "IPixelStreaming2Module.h"
+#include "IPixelStreaming2InputHandler.h"
 #include "VideoProducerMediaCapture.h"
 #include "Engine/GameEngine.h"
 #include "Slate/SceneViewport.h"
diff --git a/Engine/Plugins/Media/PixelStreaming2/Source/PixelStreaming2/Private/Logging.h b/Engine/Plugins/Media/PixelStreaming2/Source/PixelStreaming2/Private/Logging.h
--- a/Engine/Plugins/Media/PixelStreaming2/Source/PixelStreaming2/Private/Logging.h
+++ b/Engine/Plugins/Media/PixelStreaming2/Source/PixelStreaming2/Private/Logging.h
@@ -2,7 +2,6 @@
 
 #pragma once
 
-#include "HAL/IConsoleManager.h"
 #include "Logging/LogMacros.h"
 #include "Logging/StructuredLog.h"
 
diff --git a/Engine/Plugins/Media/PixelStreaming2/Source/PixelStreaming2/Private/PixelStreaming2MediaIOCapture.cpp b/Engine/Plugins/Media/PixelStreaming2/Source/PixelStreaming2/Private/PixelStreaming2MediaIOCapture.cpp
--- a/Engine/Plugins/Media/PixelStreaming2/Source/PixelStreaming2/Private/PixelStreaming2MediaIOCapture.cpp
+++ b/Engine/Plugins/Media/PixelStreaming2/Source/PixelStreaming2/Private/PixelStreaming2MediaIOCapture.cpp
@@ -5,11 +5,7 @@
 #include "Logging.h"
 #include "MediaShaders.h"
 #include "PixelCaptureInputFrameRHI.h"
-#include "RenderGraphUtils.h"
-#include "RHI.h"
 #include "ScreenPass.h"
-#include "ScreenRendering.h"
-#include "Slate/SceneViewport.h"
 
 void UPixelStreaming2MediaIOCapture::OnRHIResourceCaptured_RenderingThread(
 	FRHICommandListImmediate& /* RHICmdList */,
diff --git a/Engine/Plugins/Media/PixelStreaming2/Source/PixelStreaming2/Private/PixelStreaming2Module.cpp b/Engine/Plugins/Media/PixelStreaming2/Source/PixelStreaming2/Private/PixelStreaming2Module.cpp
--- a/Engine/Plugins/Media/PixelStreaming2/Source/PixelStreaming2/Private/PixelStreaming2Module.cpp
+++ b/Engine/Plugins/Media/PixelStreaming2/Source/PixelStreaming2/Private/PixelStreaming2Module.cpp
@@ -9,7 +9,6 @@
 #include "PixelStreaming2Common.h"
 #include "PixelStreaming2Delegates.h"
 #include "PixelStreaming2PluginSettings.h"
-#include "RHI.h"
 #include "Thread.h"
 #include "UtilsCoder.h"
 #include "UtilsCore.h"
diff --git a/Engine/Plugins/Media/PixelStreaming2/Source/PixelStreaming2/Private/VideoProducerBackBuffer.cpp b/Engine/Plugins/Media/PixelStreaming2/Source/PixelStreaming2/Private/VideoProducerBackBuffer.cpp
--- a/Engine/Plugins/Media/PixelStreaming2/Source/PixelStreaming2/Private/VideoProducerBackBuffer.cpp
+++ b/Engine/Plugins/Media/PixelStreaming2/Source/PixelStreaming2/Private/VideoProducerBackBuffer.cpp
@@ -2,14 +2,9 @@
 
 #include "VideoProducerBackBuffer.h"
 
-#include "Framework/Application/SlateApplication.h"
-#include "Logging.h"
-#include "PixelCaptureBufferFormat.h"
 #include "PixelCaptureInputFrameRHI.h"
-#include "RHI.h"
 #include "UtilsAsync.h"
-#include "Widgets/SWindow.h"
-
+#include "Framework/Application/SlateApplication.h"
 
 namespace UE::PixelStreaming2
 {
diff --git a/Engine/Plugins/Media/PixelStreaming2/Source/PixelStreaming2/Private/VideoProducerMediaCapture.cpp b/Engine/Plugins/Media/PixelStreaming2/Source/PixelStreaming2/Private/VideoProducerMediaCapture.cpp
--- a/Engine/Plugins/Media/PixelStreaming2/Source/PixelStreaming2/Private/VideoProducerMediaCapture.cpp
+++ b/Engine/Plugins/Media/PixelStreaming2/Source/PixelStreaming2/Private/VideoProducerMediaCapture.cpp
@@ -2,7 +2,6 @@
 
 #include "VideoProducerMediaCapture.h"
 
-#include "Engine/GameEngine.h"
 #include "Logging.h"
 
 namespace UE::PixelStreaming2
diff --git a/Engine/Plugins/Media/PixelStreaming2/Source/PixelStreaming2/Private/VideoProducerPIEViewport.cpp b/Engine/Plugins/Media/PixelStreaming2/Source/PixelStreaming2/Private/VideoProducerPIEViewport.cpp
--- a/Engine/Plugins/Media/PixelStreaming2/Source/PixelStreaming2/Private/VideoProducerPIEViewport.cpp
+++ b/Engine/Plugins/Media/PixelStreaming2/Source/PixelStreaming2/Private/VideoProducerPIEViewport.cpp
@@ -4,9 +4,9 @@
 
 #include "Engine/GameViewportClient.h"
 #include "PixelCaptureInputFrameRHI.h"
-#include "RenderingThread.h"
-#include "UnrealClient.h"
 #include "UtilsAsync.h"
+#include "Rendering/SlateRenderer.h"
+#include "UnrealClient.h"
 
 namespace UE::PixelStreaming2
 {
diff --git a/Engine/Plugins/Media/PixelStreaming2/Source/PixelStreaming2/Private/VideoProducerRenderTarget.cpp b/Engine/Plugins/Media/PixelStreaming2/Source/PixelStreaming2/Private/VideoProducerRenderTarget.cpp
--- a/Engine/Plugins/Media/PixelStreaming2/Source/PixelStreaming2/Private/VideoProducerRenderTarget.cpp
+++ b/Engine/Plugins/Media/PixelStreaming2/Source/PixelStreaming2/Private/VideoProducerRenderTarget.cpp
@@ -4,9 +4,8 @@
 
 #include "Engine/TextureRenderTarget2D.h"
 #include "PixelCaptureInputFrameRHI.h"
-#include "RHI.h"
+#include "Rendering/SlateRenderer.h"
 #include "TextureResource.h"
-#include "Widgets/SWindow.h"
 
 namespace UE::PixelStreaming2
 {
diff --git a/Engine/Plugins/Media/PixelStreaming2/Source/PixelStreaming2/Private/VideoSink.cpp b/Engine/Plugins/Media/PixelStreaming2/Source/PixelStreaming2/Private/VideoSink.cpp
--- a/Engine/Plugins/Media/PixelStreaming2/Source/PixelStreaming2/Private/VideoSink.cpp
+++ b/Engine/Plugins/Media/PixelStreaming2/Source/PixelStreaming2/Private/VideoSink.cpp
@@ -6,6 +6,8 @@
 #include "IPixelStreaming2VideoConsumer.h"
 #include "PixelStreaming2Trace.h"
 #include "RenderTargetPool.h"
+#include "TextureResource.h"
+#include "Framework/Application/SlateApplication.h"
 
 namespace UE::PixelStreaming2
 {
diff --git a/Engine/Plugins/Media/PixelStreaming2/Source/PixelStreaming2/Private/VideoSourceGroup.cpp b/Engine/Plugins/Media/PixelStreaming2/Source/PixelStreaming2/Private/VideoSourceGroup.cpp
--- a/Engine/Plugins/Media/PixelStreaming2/Source/PixelStreaming2/Private/VideoSourceGroup.cpp
+++ b/Engine/Plugins/Media/PixelStreaming2/Source/PixelStreaming2/Private/VideoSourceGroup.cpp
@@ -6,6 +6,10 @@
 #include "PixelStreaming2PluginSettings.h"
 #include "PixelStreaming2Trace.h"
 #include "VideoSource.h"
+#include "HAL/Runnable.h"
+#include "HAL/RunnableThread.h"
+#include "Misc/SingleThreadRunnable.h"
+#include "UnrealClient.h"
 
 namespace UE::PixelStreaming2
 {
diff --git a/Engine/Plugins/Media/PixelStreaming2/Source/PixelStreaming2/Public/IPixelStreaming2Module.h b/Engine/Plugins/Media/PixelStreaming2/Source/PixelStreaming2/Public/IPixelStreaming2Module.h
--- a/Engine/Plugins/Media/PixelStreaming2/Source/PixelStreaming2/Public/IPixelStreaming2Module.h
+++ b/Engine/Plugins/Media/PixelStreaming2/Source/PixelStreaming2/Public/IPixelStreaming2Module.h
@@ -8,7 +8,6 @@
 #include "Templates/RefCounting.h"
 #include "Templates/SharedPointer.h"
 #include "IPixelStreaming2AudioProducer.h"
-#include "IPixelStreaming2InputHandler.h"
 #include "IPixelStreaming2Streamer.h"
 #include "IPixelStreaming2VideoProducer.h"
 
diff --git a/Engine/Plugins/Media/PixelStreaming2/Source/PixelStreaming2Core/Private/Logging.h b/Engine/Plugins/Media/PixelStreaming2/Source/PixelStreaming2Core/Private/Logging.h
--- a/Engine/Plugins/Media/PixelStreaming2/Source/PixelStreaming2Core/Private/Logging.h
+++ b/Engine/Plugins/Media/PixelStreaming2/Source/PixelStreaming2Core/Private/Logging.h
@@ -2,7 +2,6 @@
 
 #pragma once
 
-#include "HAL/IConsoleManager.h"
 #include "Logging/LogMacros.h"
 #include "Logging/StructuredLog.h"
 
diff --git a/Engine/Plugins/Media/PixelStreaming2/Source/PixelStreaming2Core/Public/IPixelStreaming2AudioConsumer.h b/Engine/Plugins/Media/PixelStreaming2/Source/PixelStreaming2Core/Public/IPixelStreaming2AudioConsumer.h
--- a/Engine/Plugins/Media/PixelStreaming2/Source/PixelStreaming2Core/Public/IPixelStreaming2AudioConsumer.h
+++ b/Engine/Plugins/Media/PixelStreaming2/Source/PixelStreaming2Core/Public/IPixelStreaming2AudioConsumer.h
@@ -2,7 +2,6 @@
 
 #pragma once
 
-#include "CoreTypes.h"
 #include "UObject/Interface.h"
 
 #include "IPixelStreaming2AudioConsumer.generated.h"
diff --git a/Engine/Plugins/Media/PixelStreaming2/Source/PixelStreaming2Core/Public/IPixelStreaming2AudioProducer.h b/Engine/Plugins/Media/PixelStreaming2/Source/PixelStreaming2Core/Public/IPixelStreaming2AudioProducer.h
--- a/Engine/Plugins/Media/PixelStreaming2/Source/PixelStreaming2Core/Public/IPixelStreaming2AudioProducer.h
+++ b/Engine/Plugins/Media/PixelStreaming2/Source/PixelStreaming2Core/Public/IPixelStreaming2AudioProducer.h
@@ -2,7 +2,6 @@
 
 #pragma once
 
-#include "CoreTypes.h"
 #include "UObject/Interface.h"
 
 #include "IPixelStreaming2AudioProducer.generated.h"
diff --git a/Engine/Plugins/Media/PixelStreaming2/Source/PixelStreaming2Core/Public/IPixelStreaming2AudioSink.h b/Engine/Plugins/Media/PixelStreaming2/Source/PixelStreaming2Core/Public/IPixelStreaming2AudioSink.h
--- a/Engine/Plugins/Media/PixelStreaming2/Source/PixelStreaming2Core/Public/IPixelStreaming2AudioSink.h
+++ b/Engine/Plugins/Media/PixelStreaming2/Source/PixelStreaming2Core/Public/IPixelStreaming2AudioSink.h
@@ -2,7 +2,6 @@
 
 #pragma once
 
-#include "CoreTypes.h"
 #include "IPixelStreaming2AudioConsumer.h"
 #include "Templates/PointerVariants.h"
 #include "UObject/Interface.h"
diff --git a/Engine/Plugins/Media/PixelStreaming2/Source/PixelStreaming2Core/Public/IPixelStreaming2Streamer.h b/Engine/Plugins/Media/PixelStreaming2/Source/PixelStreaming2Core/Public/IPixelStreaming2Streamer.h
--- a/Engine/Plugins/Media/PixelStreaming2/Source/PixelStreaming2Core/Public/IPixelStreaming2Streamer.h
+++ b/Engine/Plugins/Media/PixelStreaming2/Source/PixelStreaming2Core/Public/IPixelStreaming2Streamer.h
@@ -2,13 +2,11 @@
 
 #pragma once
 
-#include "CoreTypes.h"
 #include "Features/IModularFeature.h"
 #include "IPixelStreaming2AudioSink.h"
 #include "IPixelStreaming2VideoProducer.h"
 #include "IPixelStreaming2VideoSink.h"
 #include "Templates/SharedPointer.h"
-#include "UObject/Interface.h"
 
 class UTexture2D;
 namespace UE::PixelStreaming2
diff --git a/Engine/Plugins/Media/PixelStreaming2/Source/PixelStreaming2Core/Public/IPixelStreaming2VideoConsumer.h b/Engine/Plugins/Media/PixelStreaming2/Source/PixelStreaming2Core/Public/IPixelStreaming2VideoConsumer.h
--- a/Engine/Plugins/Media/PixelStreaming2/Source/PixelStreaming2Core/Public/IPixelStreaming2VideoConsumer.h
+++ b/Engine/Plugins/Media/PixelStreaming2/Source/PixelStreaming2Core/Public/IPixelStreaming2VideoConsumer.h
@@ -2,7 +2,6 @@
 
 #pragma once
 
-#include "CoreTypes.h"
 #include "UObject/Interface.h"
 #include "RHIFwd.h"
 
diff --git a/Engine/Plugins/Media/PixelStreaming2/Source/PixelStreaming2Core/Public/IPixelStreaming2VideoProducer.h b/Engine/Plugins/Media/PixelStreaming2/Source/PixelStreaming2Core/Public/IPixelStreaming2VideoProducer.h
--- a/Engine/Plugins/Media/PixelStreaming2/Source/PixelStreaming2Core/Public/IPixelStreaming2VideoProducer.h
+++ b/Engine/Plugins/Media/PixelStreaming2/Source/PixelStreaming2Core/Public/IPixelStreaming2VideoProducer.h
@@ -2,7 +2,6 @@
 
 #pragma once
 
-#include "CoreTypes.h"
 #include "IPixelCaptureInputFrame.h"
 #include "UObject/Interface.h"
 
diff --git a/Engine/Plugins/Media/PixelStreaming2/Source/PixelStreaming2Core/Public/IPixelStreaming2VideoSink.h b/Engine/Plugins/Media/PixelStreaming2/Source/PixelStreaming2Core/Public/IPixelStreaming2VideoSink.h
--- a/Engine/Plugins/Media/PixelStreaming2/Source/PixelStreaming2Core/Public/IPixelStreaming2VideoSink.h
+++ b/Engine/Plugins/Media/PixelStreaming2/Source/PixelStreaming2Core/Public/IPixelStreaming2VideoSink.h
@@ -2,7 +2,6 @@
 
 #pragma once
 
-#include "CoreTypes.h"
 #include "IPixelStreaming2VideoConsumer.h"
 #include "Templates/PointerVariants.h"
 #include "UObject/Interface.h"
diff --git a/Engine/Plugins/Media/PixelStreaming2/Source/PixelStreaming2Editor/Internal/PixelStreaming2Style.h b/Engine/Plugins/Media/PixelStreaming2/Source/PixelStreaming2Editor/Internal/PixelStreaming2Style.h
--- a/Engine/Plugins/Media/PixelStreaming2/Source/PixelStreaming2Editor/Internal/PixelStreaming2Style.h
+++ b/Engine/Plugins/Media/PixelStreaming2/Source/PixelStreaming2Editor/Internal/PixelStreaming2Style.h
@@ -1,7 +1,6 @@
 // Copyright Epic Games, Inc. All Rights Reserved.
 
 #pragma once
-#include "CoreMinimal.h"
 #include "Styling/SlateStyle.h"
 
 namespace UE::EditorPixelStreaming2
diff --git a/Engine/Plugins/Media/PixelStreaming2/Source/PixelStreaming2Editor/Internal/VideoProducerViewport.h b/Engine/Plugins/Media/PixelStreaming2/Source/PixelStreaming2Editor/Internal/VideoProducerViewport.h
--- a/Engine/Plugins/Media/PixelStreaming2/Source/PixelStreaming2Editor/Internal/VideoProducerViewport.h
+++ b/Engine/Plugins/Media/PixelStreaming2/Source/PixelStreaming2Editor/Internal/VideoProducerViewport.h
@@ -2,11 +2,10 @@
 
 #pragma once
 
-#include "Delegates/IDelegateInstance.h"
 #include "IPixelStreaming2Streamer.h"
-#include "UnrealClient.h"
 #include "VideoProducer.h"
 #include "Widgets/SWindow.h"
+class FViewport;
 
 namespace UE::EditorPixelStreaming2
 {
diff --git a/Engine/Plugins/Media/PixelStreaming2/Source/PixelStreaming2Editor/Private/AssetTypeActions_VideoProducer.h b/Engine/Plugins/Media/PixelStreaming2/Source/PixelStreaming2Editor/Private/AssetTypeActions_VideoProducer.h
--- a/Engine/Plugins/Media/PixelStreaming2/Source/PixelStreaming2Editor/Private/AssetTypeActions_VideoProducer.h
+++ b/Engine/Plugins/Media/PixelStreaming2/Source/PixelStreaming2Editor/Private/AssetTypeActions_VideoProducer.h
@@ -4,8 +4,6 @@
 
 #include "AssetTypeActions_Base.h"
 #include "Blueprints/PixelStreaming2VideoProducer.h"
-#include "CoreMinimal.h"
-#include "Toolkits/IToolkitHost.h"
 
 class FAssetTypeActions_VideoProducer : public FAssetTypeActions_Base
 {
diff --git a/Engine/Plugins/Media/PixelStreaming2/Source/PixelStreaming2Editor/Private/Factories.cpp b/Engine/Plugins/Media/PixelStreaming2/Source/PixelStreaming2Editor/Private/Factories.cpp
--- a/Engine/Plugins/Media/PixelStreaming2/Source/PixelStreaming2Editor/Private/Factories.cpp
+++ b/Engine/Plugins/Media/PixelStreaming2/Source/PixelStreaming2Editor/Private/Factories.cpp
@@ -7,7 +7,6 @@
 #include "Blueprints/PixelStreaming2MediaTexture.h"
 #include "Blueprints/PixelStreaming2VideoProducer.h"
 #include "IAssetTools.h"
-#include "Modules/ModuleManager.h"
 
 #define LOCTEXT_NAMESPACE "PixelStreaming2"
 
diff --git a/Engine/Plugins/Media/PixelStreaming2/Source/PixelStreaming2Editor/Private/Factories.h b/Engine/Plugins/Media/PixelStreaming2/Source/PixelStreaming2Editor/Private/Factories.h
--- a/Engine/Plugins/Media/PixelStreaming2/Source/PixelStreaming2Editor/Private/Factories.h
+++ b/Engine/Plugins/Media/PixelStreaming2/Source/PixelStreaming2Editor/Private/Factories.h
@@ -2,9 +2,7 @@
 
 #pragma once
 
-#include "CoreMinimal.h"
 #include "Factories/Factory.h"
-#include "UObject/Object.h"
 
 #include "Factories.generated.h"
 
diff --git a/Engine/Plugins/Media/PixelStreaming2/Source/PixelStreaming2Editor/Private/Logging.h b/Engine/Plugins/Media/PixelStreaming2/Source/PixelStreaming2Editor/Private/Logging.h
--- a/Engine/Plugins/Media/PixelStreaming2/Source/PixelStreaming2Editor/Private/Logging.h
+++ b/Engine/Plugins/Media/PixelStreaming2/Source/PixelStreaming2Editor/Private/Logging.h
@@ -2,7 +2,6 @@
 
 #pragma once
 
-#include "HAL/IConsoleManager.h"
 #include "Logging/LogMacros.h"
 #include "Logging/StructuredLog.h"
 
diff --git a/Engine/Plugins/Media/PixelStreaming2/Source/PixelStreaming2Editor/Private/PixelStreaming2Commands.h b/Engine/Plugins/Media/PixelStreaming2/Source/PixelStreaming2Editor/Private/PixelStreaming2Commands.h
--- a/Engine/Plugins/Media/PixelStreaming2/Source/PixelStreaming2Editor/Private/PixelStreaming2Commands.h
+++ b/Engine/Plugins/Media/PixelStreaming2/Source/PixelStreaming2Editor/Private/PixelStreaming2Commands.h
@@ -2,7 +2,6 @@
 
 #pragma once
 
-#include "CoreMinimal.h"
 #include "Framework/Commands/Commands.h"
 
 /**
diff --git a/Engine/Plugins/Media/PixelStreaming2/Source/PixelStreaming2Editor/Private/PixelStreaming2EditorModule.cpp b/Engine/Plugins/Media/PixelStreaming2/Source/PixelStreaming2Editor/Private/PixelStreaming2EditorModule.cpp
--- a/Engine/Plugins/Media/PixelStreaming2/Source/PixelStreaming2Editor/Private/PixelStreaming2EditorModule.cpp
+++ b/Engine/Plugins/Media/PixelStreaming2/Source/PixelStreaming2Editor/Private/PixelStreaming2EditorModule.cpp
@@ -8,30 +8,26 @@
 #include "Engine/GameViewportClient.h"
 #include "UtilsCore.h"
 #include "Editor/EditorPerformanceSettings.h"
-#include "Framework/Application/SlateApplication.h"
-#include "IPixelStreaming2RTCModule.h"
-#include "IPixelStreaming2Streamer.h"
 #include "VideoProducerBackBufferComposited.h"
 #include "VideoProducerPIEViewport.h"
 #include "VideoProducerViewport.h"
 #include "PixelStreaming2PluginSettings.h"
 #include "UnrealEngine.h"
+#include "Engine/Engine.h"
 #include "Interfaces/IMainFrameModule.h"
 #include "PixelStreaming2Utils.h"
 #include "GenericPlatform/GenericPlatformHttp.h"
 #include "LevelEditorViewport.h"
 #include "Slate/SceneViewport.h"
 #include "SLevelViewport.h"
-#include "PixelStreaming2InputEnums.h"
-#include "IPixelStreaming2InputModule.h"
 #include "IPixelStreaming2InputHandler.h"
-#include "IPixelStreaming2DataProtocol.h"
 #include "PixelStreaming2Delegates.h"
 #include "Framework/Notifications/NotificationManager.h"
 #include "Widgets/Notifications/SNotificationList.h"
 #include "Logging.h"
 #include "IPixelStreaming2Module.h"
 #include "UtilsAsync.h"
+#include "LevelEditor.h"
 
 #define LOCTEXT_NAMESPACE "PixelStreaming2EditorModule"
 
diff --git a/Engine/Plugins/Media/PixelStreaming2/Source/PixelStreaming2Editor/Private/PixelStreaming2EditorModule.h b/Engine/Plugins/Media/PixelStreaming2/Source/PixelStreaming2Editor/Private/PixelStreaming2EditorModule.h
--- a/Engine/Plugins/Media/PixelStreaming2/Source/PixelStreaming2Editor/Private/PixelStreaming2EditorModule.h
+++ b/Engine/Plugins/Media/PixelStreaming2/Source/PixelStreaming2Editor/Private/PixelStreaming2EditorModule.h
@@ -4,8 +4,6 @@
 
 #include "IPixelStreaming2EditorModule.h"
 #include "IPixelStreaming2RTCModule.h"
-#include "PixelStreaming2Servers.h"
-#include "Widgets/SWindow.h"
 
 #include <atomic>
 
diff --git a/Engine/Plugins/Media/PixelStreaming2/Source/PixelStreaming2Editor/Private/PixelStreaming2Style.cpp b/Engine/Plugins/Media/PixelStreaming2/Source/PixelStreaming2Editor/Private/PixelStreaming2Style.cpp
--- a/Engine/Plugins/Media/PixelStreaming2/Source/PixelStreaming2Editor/Private/PixelStreaming2Style.cpp
+++ b/Engine/Plugins/Media/PixelStreaming2/Source/PixelStreaming2Editor/Private/PixelStreaming2Style.cpp
@@ -3,9 +3,7 @@
 #include "PixelStreaming2Style.h"
 #include "Styling/SlateStyleRegistry.h"
 #include "Framework/Application/SlateApplication.h"
-#include "Slate/SlateGameResources.h"
 #include "Interfaces/IPluginManager.h"
-#include "Styling/SlateStyle.h"
 #include "Styling/SlateStyleMacros.h"
 
 #define RootToContentDir Style->RootToContentDir
diff --git a/Engine/Plugins/Media/PixelStreaming2/Source/PixelStreaming2Editor/Private/PixelStreaming2Toolbar.cpp b/Engine/Plugins/Media/PixelStreaming2/Source/PixelStreaming2Editor/Private/PixelStreaming2Toolbar.cpp
--- a/Engine/Plugins/Media/PixelStreaming2/Source/PixelStreaming2Editor/Private/PixelStreaming2Toolbar.cpp
+++ b/Engine/Plugins/Media/PixelStreaming2/Source/PixelStreaming2Editor/Private/PixelStreaming2Toolbar.cpp
@@ -1,36 +1,15 @@
 // Copyright Epic Games, Inc. All Rights Reserved.
 
 #include "PixelStreaming2Toolbar.h"
-#include "Framework/Commands/UIAction.h"
-#include "Framework/Commands/UICommandList.h"
-#include "Framework/Docking/TabManager.h"
-#include "LevelEditor.h"
 #include "ToolMenus.h"
 #include "PixelStreaming2Commands.h"
 #include "PixelStreaming2Style.h"
-#include "Framework/SlateDelegates.h"
-#include "ToolMenuContext.h"
 #include "IPixelStreaming2Module.h"
-#include "IPixelStreaming2Streamer.h"
-#include "Editor/EditorEngine.h"
-#include "Slate/SceneViewport.h"
-#include "Widgets/SWindow.h"
-#include "EditorViewportClient.h"
-#include "Slate/SceneViewport.h"
 #include "Widgets/SViewport.h"
-#include "Widgets/Layout/SSpacer.h"
-#include "Styling/AppStyle.h"
 #include "PixelStreaming2EditorModule.h"
 #include "Widgets/Input/SEditableTextBox.h"
-#include "Widgets/Input/SSpinBox.h"
-#include "Types/SlateEnums.h"
-#include "Input/Events.h"
-#include "InputCoreTypes.h"
-#include "PixelStreaming2Servers.h"
-#include "GenericPlatform/GenericPlatformHttp.h"
+#include "Widgets/Input/SButton.h"
 #include <SocketSubsystem.h>
-#include <IPAddress.h>
-#include "SlateFwd.h"
 #include "Widgets/Input/SNumericEntryBox.h"
 #include "PixelStreaming2PluginSettings.h"
 #include "Logging.h"
diff --git a/Engine/Plugins/Media/PixelStreaming2/Source/PixelStreaming2Editor/Private/PixelStreaming2Toolbar.h b/Engine/Plugins/Media/PixelStreaming2/Source/PixelStreaming2Editor/Private/PixelStreaming2Toolbar.h
--- a/Engine/Plugins/Media/PixelStreaming2/Source/PixelStreaming2Editor/Private/PixelStreaming2Toolbar.h
+++ b/Engine/Plugins/Media/PixelStreaming2/Source/PixelStreaming2Editor/Private/PixelStreaming2Toolbar.h
@@ -2,12 +2,8 @@
 
 #pragma once
 
-#include "Framework/Commands/UIAction.h"
 #include "Framework/Commands/UICommandList.h"
-#include "Framework/Docking/TabManager.h"
 #include "Framework/MultiBox/MultiBoxBuilder.h"
-#include "LevelEditor.h"
-#include "ToolMenus.h"
 
 namespace UE::EditorPixelStreaming2
 {
diff --git a/Engine/Plugins/Media/PixelStreaming2/Source/PixelStreaming2Editor/Private/VideoProducerBackBufferComposited.cpp b/Engine/Plugins/Media/PixelStreaming2/Source/PixelStreaming2Editor/Private/VideoProducerBackBufferComposited.cpp
--- a/Engine/Plugins/Media/PixelStreaming2/Source/PixelStreaming2Editor/Private/VideoProducerBackBufferComposited.cpp
+++ b/Engine/Plugins/Media/PixelStreaming2/Source/PixelStreaming2Editor/Private/VideoProducerBackBufferComposited.cpp
@@ -7,12 +7,7 @@
 #include "Framework/Application/SlateApplication.h"
 #include "MediaShaders.h"
 #include "PixelCaptureInputFrameRHI.h"
-#include "RenderGraphBuilder.h"
-#include "RenderGraphUtils.h"
-#include "RHI.h"
 #include "ScreenPass.h"
-#include "ScreenRendering.h"
-#include "Utils.h"
 
 namespace UE::EditorPixelStreaming2
 {
diff --git a/Engine/Plugins/Media/PixelStreaming2/Source/PixelStreaming2Editor/Private/VideoProducerViewport.cpp b/Engine/Plugins/Media/PixelStreaming2/Source/PixelStreaming2Editor/Private/VideoProducerViewport.cpp
--- a/Engine/Plugins/Media/PixelStreaming2/Source/PixelStreaming2Editor/Private/VideoProducerViewport.cpp
+++ b/Engine/Plugins/Media/PixelStreaming2/Source/PixelStreaming2Editor/Private/VideoProducerViewport.cpp
@@ -2,21 +2,15 @@
 
 #include "VideoProducerViewport.h"
 
-#include "Async/Async.h"
 #include "Engine/GameViewportClient.h"
-#include "Framework/Application/SlateApplication.h"
 #include "IPixelStreaming2InputHandler.h"
 #include "LevelEditor.h"
 #include "LevelEditorViewport.h"
 #include "Modules/ModuleManager.h"
-#include "PixelCaptureBufferFormat.h"
-#include "PixelCaptureCapturerRHI.h"
 #include "PixelCaptureInputFrameRHI.h"
-#include "RenderingThread.h"
 #include "SLevelViewport.h"
-#include "UnrealClient.h"
-#include "Utils.h"
-
+#include "Framework/Application/SlateApplication.h"
+#include "Rendering/SlateRenderer.h"
 
 namespace UE::EditorPixelStreaming2
 {
diff --git a/Engine/Plugins/Media/PixelStreaming2/Source/PixelStreaming2HMD/PixelStreaming2HMD.build.cs b/Engine/Plugins/Media/PixelStreaming2/Source/PixelStreaming2HMD/PixelStreaming2HMD.build.cs
--- a/Engine/Plugins/Media/PixelStreaming2/Source/PixelStreaming2HMD/PixelStreaming2HMD.build.cs
+++ b/Engine/Plugins/Media/PixelStreaming2/Source/PixelStreaming2HMD/PixelStreaming2HMD.build.cs
@@ -22,6 +22,12 @@ namespace UnrealBuildTool.Rules
 					"PixelStreaming2Settings"
 				}
 			);
+			PublicDependencyModuleNames.AddRange(
+				new string[]
+				{
+					"Core"
+				}
+			);
 		}
 	}
 }
\ No newline at end of file
diff --git a/Engine/Plugins/Media/PixelStreaming2/Source/PixelStreaming2HMD/Private/Logging.h b/Engine/Plugins/Media/PixelStreaming2/Source/PixelStreaming2HMD/Private/Logging.h
--- a/Engine/Plugins/Media/PixelStreaming2/Source/PixelStreaming2HMD/Private/Logging.h
+++ b/Engine/Plugins/Media/PixelStreaming2/Source/PixelStreaming2HMD/Private/Logging.h
@@ -2,8 +2,6 @@
 
 #pragma once
 
-#include "HAL/IConsoleManager.h"
 #include "Logging/LogMacros.h"
-#include "Logging/StructuredLog.h"
 
 DECLARE_LOG_CATEGORY_EXTERN(LogPixelStreaming2HMD, Log, All);
diff --git a/Engine/Plugins/Media/PixelStreaming2/Source/PixelStreaming2HMD/Private/PixelStreaming2HMD.cpp b/Engine/Plugins/Media/PixelStreaming2/Source/PixelStreaming2HMD/Private/PixelStreaming2HMD.cpp
--- a/Engine/Plugins/Media/PixelStreaming2/Source/PixelStreaming2HMD/Private/PixelStreaming2HMD.cpp
+++ b/Engine/Plugins/Media/PixelStreaming2/Source/PixelStreaming2HMD/Private/PixelStreaming2HMD.cpp
@@ -2,17 +2,10 @@
 
 #include "PixelStreaming2HMD.h"
 
-#include "EngineGlobals.h"
-#include "Engine/Engine.h"
-#include "Engine/LocalPlayer.h"
 #include "Engine/GameViewportClient.h"
-#include "GameFramework/PlayerController.h"
 #include "GameFramework/WorldSettings.h"
 #include "Logging.h"
-#include "Misc/App.h"
-#include "Modules/ModuleManager.h"
 #include "PixelStreaming2PluginSettings.h"
-#include "PostProcess/PostProcessHMD.h"
 #include "SceneView.h"
 #include "Widgets/SWindow.h"
 
@@ -331,4 +324,4 @@ void FPixelStreaming2HMD::SetupViewFamily(FSceneViewFamily& InViewFamily)
 	{
 		WorldToMeters = World->GetWorldSettings()->WorldToMeters;
 	}
-}
\ No newline at end of file
+}
diff --git a/Engine/Plugins/Media/PixelStreaming2/Source/PixelStreaming2HMD/Private/PixelStreaming2HMD.h b/Engine/Plugins/Media/PixelStreaming2/Source/PixelStreaming2HMD/Private/PixelStreaming2HMD.h
--- a/Engine/Plugins/Media/PixelStreaming2/Source/PixelStreaming2HMD/Private/PixelStreaming2HMD.h
+++ b/Engine/Plugins/Media/PixelStreaming2/Source/PixelStreaming2HMD/Private/PixelStreaming2HMD.h
@@ -2,11 +2,9 @@
 
 #pragma once
 
-#include "CoreMinimal.h"
 #include "HeadMountedDisplayBase.h"
 #include "IPixelStreaming2HMD.h"
 #include "SceneViewExtension.h"
-#include "XRTrackingSystemBase.h"
 
 class APlayerController;
 class FSceneView;
diff --git a/Engine/Plugins/Media/PixelStreaming2/Source/PixelStreaming2HMD/Private/PixelStreaming2HMDModule.cpp b/Engine/Plugins/Media/PixelStreaming2/Source/PixelStreaming2HMD/Private/PixelStreaming2HMDModule.cpp
--- a/Engine/Plugins/Media/PixelStreaming2/Source/PixelStreaming2HMD/Private/PixelStreaming2HMDModule.cpp
+++ b/Engine/Plugins/Media/PixelStreaming2/Source/PixelStreaming2HMD/Private/PixelStreaming2HMDModule.cpp
@@ -2,7 +2,6 @@
 
 #include "PixelStreaming2HMDModule.h"
 
-#include "IXRTrackingSystem.h"
 #include "PixelStreaming2HMD.h"
 #include "PixelStreaming2PluginSettings.h"
 
diff --git a/Engine/Plugins/Media/PixelStreaming2/Source/PixelStreaming2HMD/Public/IPixelStreaming2HMD.h b/Engine/Plugins/Media/PixelStreaming2/Source/PixelStreaming2HMD/Public/IPixelStreaming2HMD.h
--- a/Engine/Plugins/Media/PixelStreaming2/Source/PixelStreaming2HMD/Public/IPixelStreaming2HMD.h
+++ b/Engine/Plugins/Media/PixelStreaming2/Source/PixelStreaming2HMD/Public/IPixelStreaming2HMD.h
@@ -2,7 +2,6 @@
 
 #pragma once
 
-#include "CoreMinimal.h"
 #include "Math/MathFwd.h"
 
 /**
diff --git a/Engine/Plugins/Media/PixelStreaming2/Source/PixelStreaming2HMD/Public/IPixelStreaming2HMDModule.h b/Engine/Plugins/Media/PixelStreaming2/Source/PixelStreaming2HMD/Public/IPixelStreaming2HMDModule.h
--- a/Engine/Plugins/Media/PixelStreaming2/Source/PixelStreaming2HMD/Public/IPixelStreaming2HMDModule.h
+++ b/Engine/Plugins/Media/PixelStreaming2/Source/PixelStreaming2HMD/Public/IPixelStreaming2HMDModule.h
@@ -2,7 +2,6 @@
 
 #pragma once
 
-#include "CoreMinimal.h"
 #include "IHeadMountedDisplayModule.h"
 #include "IPixelStreaming2HMD.h"
 #include "Modules/ModuleManager.h"
diff --git a/Engine/Plugins/Media/PixelStreaming2/Source/PixelStreaming2Input/Internal/InputHandler.h b/Engine/Plugins/Media/PixelStreaming2/Source/PixelStreaming2Input/Internal/InputHandler.h
--- a/Engine/Plugins/Media/PixelStreaming2/Source/PixelStreaming2Input/Internal/InputHandler.h
+++ b/Engine/Plugins/Media/PixelStreaming2/Source/PixelStreaming2Input/Internal/InputHandler.h
@@ -2,13 +2,9 @@
 
 #pragma once
 
-#include "GenericPlatform/GenericApplicationMessageHandler.h"
 #include "Slate/SceneViewport.h"
-#include "Widgets/SWindow.h"
-#include "InputCoreTypes.h"
 #include "IPixelStreaming2InputHandler.h"
 #include "XRMotionControllerBase.h"
-#include "HeadMountedDisplayTypes.h"
 #include "PixelStreaming2HMDEnums.h"
 
 namespace UE::PixelStreaming2Input
diff --git a/Engine/Plugins/Media/PixelStreaming2/Source/PixelStreaming2Input/Internal/InputProtocolMap.h b/Engine/Plugins/Media/PixelStreaming2/Source/PixelStreaming2Input/Internal/InputProtocolMap.h
--- a/Engine/Plugins/Media/PixelStreaming2/Source/PixelStreaming2Input/Internal/InputProtocolMap.h
+++ b/Engine/Plugins/Media/PixelStreaming2/Source/PixelStreaming2Input/Internal/InputProtocolMap.h
@@ -2,9 +2,6 @@
 
 #pragma once
 
-#include "Containers/Map.h"
-#include "Containers/Set.h"
-#include "Misc/ScopeLock.h"
 #include "IPixelStreaming2InputMessage.h"
 #include "IPixelStreaming2DataProtocol.h"
 #include "PixelStreaming2InputEnums.h"
diff --git a/Engine/Plugins/Media/PixelStreaming2/Source/PixelStreaming2Input/Private/InputDevice.cpp b/Engine/Plugins/Media/PixelStreaming2/Source/PixelStreaming2Input/Private/InputDevice.cpp
--- a/Engine/Plugins/Media/PixelStreaming2/Source/PixelStreaming2Input/Private/InputDevice.cpp
+++ b/Engine/Plugins/Media/PixelStreaming2/Source/PixelStreaming2Input/Private/InputDevice.cpp
@@ -2,7 +2,6 @@
 
 #include "InputDevice.h"
 #include "Framework/Application/SlateApplication.h"
-#include "GenericPlatform/GenericPlatformInputDeviceMapper.h"
 
 TSharedPtr<UE::PixelStreaming2Input::FInputDevice> UE::PixelStreaming2Input::FInputDevice::InputDevice;
 
diff --git a/Engine/Plugins/Media/PixelStreaming2/Source/PixelStreaming2Input/Private/InputDevice.h b/Engine/Plugins/Media/PixelStreaming2/Source/PixelStreaming2Input/Private/InputDevice.h
--- a/Engine/Plugins/Media/PixelStreaming2/Source/PixelStreaming2Input/Private/InputDevice.h
+++ b/Engine/Plugins/Media/PixelStreaming2/Source/PixelStreaming2Input/Private/InputDevice.h
@@ -4,7 +4,6 @@
 
 #include "IInputDevice.h"
 #include "IPixelStreaming2InputHandler.h"
-#include "Misc/CoreMiscDefines.h"
 
 namespace UE::PixelStreaming2Input
 {
diff --git a/Engine/Plugins/Media/PixelStreaming2/Source/PixelStreaming2Input/Private/InputHandler.cpp b/Engine/Plugins/Media/PixelStreaming2/Source/PixelStreaming2Input/Private/InputHandler.cpp
--- a/Engine/Plugins/Media/PixelStreaming2/Source/PixelStreaming2Input/Private/InputHandler.cpp
+++ b/Engine/Plugins/Media/PixelStreaming2/Source/PixelStreaming2Input/Private/InputHandler.cpp
@@ -4,29 +4,18 @@
 
 #include "ApplicationWrapper.h"
 #include "InputStructures.h"
-#include "JavaScriptKeyCodes.inl"
-#include "Layout/ArrangedChildren.h"
-#include "Layout/WidgetPath.h"
 #include "Framework/Application/SlateApplication.h"
-#include "Widgets/SViewport.h"
 #include "Policies/CondensedJsonPrintPolicy.h"
 #include "Serialization/JsonSerializer.h"
-#include "Utils.h"
 #include "Engine/Engine.h"
 #include "Framework/Application/SlateUser.h"
 #include "Widgets/Input/SEditableTextBox.h"
-#include "Misc/CoreMiscDefines.h"
 #include "Input/HittestGrid.h"
-#include "IPixelStreaming2HMD.h"
 #include "IPixelStreaming2HMDModule.h"
-#include "PixelStreaming2InputEnums.h"
-#include "Widgets/Input/SEditableText.h"
 #include "Widgets/Text/SMultiLineEditableText.h"
 #include "InputDevice.h"
-#include "DefaultDataProtocol.h"
 #include "Logging.h"
 #include "PixelStreaming2PluginSettings.h"
-#include "PixelStreaming2InputModule.h"
 
 #if WITH_EDITOR
 	#include "UObject/UObjectGlobals.h"
diff --git a/Engine/Plugins/Media/PixelStreaming2/Source/PixelStreaming2Input/Private/InputProtocolMap.cpp b/Engine/Plugins/Media/PixelStreaming2/Source/PixelStreaming2Input/Private/InputProtocolMap.cpp
--- a/Engine/Plugins/Media/PixelStreaming2/Source/PixelStreaming2Input/Private/InputProtocolMap.cpp
+++ b/Engine/Plugins/Media/PixelStreaming2/Source/PixelStreaming2Input/Private/InputProtocolMap.cpp
@@ -1,7 +1,6 @@
 // Copyright Epic Games, Inc. All Rights Reserved.
 
 #include "InputProtocolMap.h"
-#include "DefaultDataProtocol.h"
 #include "InputMessage.h"
 
 namespace UE::PixelStreaming2Input
diff --git a/Engine/Plugins/Media/PixelStreaming2/Source/PixelStreaming2Input/Private/Logging.h b/Engine/Plugins/Media/PixelStreaming2/Source/PixelStreaming2Input/Private/Logging.h
--- a/Engine/Plugins/Media/PixelStreaming2/Source/PixelStreaming2Input/Private/Logging.h
+++ b/Engine/Plugins/Media/PixelStreaming2/Source/PixelStreaming2Input/Private/Logging.h
@@ -2,8 +2,6 @@
 
 #pragma once
 
-#include "HAL/IConsoleManager.h"
 #include "Logging/LogMacros.h"
-#include "Logging/StructuredLog.h"
 
 DECLARE_LOG_CATEGORY_EXTERN(LogPixelStreaming2Input, Log, All);
\ No newline at end of file
diff --git a/Engine/Plugins/Media/PixelStreaming2/Source/PixelStreaming2Input/Private/PixelStreaming2InputModule.cpp b/Engine/Plugins/Media/PixelStreaming2/Source/PixelStreaming2Input/Private/PixelStreaming2InputModule.cpp
--- a/Engine/Plugins/Media/PixelStreaming2/Source/PixelStreaming2Input/Private/PixelStreaming2InputModule.cpp
+++ b/Engine/Plugins/Media/PixelStreaming2/Source/PixelStreaming2Input/Private/PixelStreaming2InputModule.cpp
@@ -2,7 +2,6 @@
 
 #include "PixelStreaming2InputModule.h"
 
-#include "ApplicationWrapper.h"
 #include "Framework/Application/SlateApplication.h"
 #include "InputHandler.h"
 #include "IPixelStreaming2HMDModule.h"
diff --git a/Engine/Plugins/Media/PixelStreaming2/Source/PixelStreaming2Input/Public/IPixelStreaming2InputHandler.h b/Engine/Plugins/Media/PixelStreaming2/Source/PixelStreaming2Input/Public/IPixelStreaming2InputHandler.h
--- a/Engine/Plugins/Media/PixelStreaming2/Source/PixelStreaming2Input/Public/IPixelStreaming2InputHandler.h
+++ b/Engine/Plugins/Media/PixelStreaming2/Source/PixelStreaming2Input/Public/IPixelStreaming2InputHandler.h
@@ -2,15 +2,11 @@
 
 #pragma once
 
-#include "GenericPlatform/GenericPlatformMisc.h"
+//#include "GenericPlatform/GenericPlatformMisc.h"
 #include "IInputDevice.h"
 #include "IPixelStreaming2DataProtocol.h"
 #include "PixelStreaming2HMDEnums.h"
-#include "PixelStreaming2InputEnums.h"
-#include "Serialization/MemoryReader.h"
-#include "Templates/SharedPointer.h"
 #include "Widgets/SViewport.h"
-#include "Widgets/SWindow.h"
 
 /**
  * The IPixelStreaming2InputHandler, used to handle input from a remote peer and pass it to UE accordingly. Setting the target viewport allows for
diff --git a/Engine/Plugins/Media/PixelStreaming2/Source/PixelStreaming2Input/Public/IPixelStreaming2InputMessage.h b/Engine/Plugins/Media/PixelStreaming2/Source/PixelStreaming2Input/Public/IPixelStreaming2InputMessage.h
--- a/Engine/Plugins/Media/PixelStreaming2/Source/PixelStreaming2Input/Public/IPixelStreaming2InputMessage.h
+++ b/Engine/Plugins/Media/PixelStreaming2/Source/PixelStreaming2Input/Public/IPixelStreaming2InputMessage.h
@@ -5,7 +5,6 @@
 #include "PixelStreaming2InputEnums.h"
 
 #include "Containers/Array.h"
-#include "HAL/Platform.h"
 
 /**
  * A message that is sent over the Pixel Streaming data channel.
diff --git a/Engine/Plugins/Media/PixelStreaming2/Source/PixelStreaming2Input/Public/IPixelStreaming2InputModule.h b/Engine/Plugins/Media/PixelStreaming2/Source/PixelStreaming2Input/Public/IPixelStreaming2InputModule.h
--- a/Engine/Plugins/Media/PixelStreaming2/Source/PixelStreaming2Input/Public/IPixelStreaming2InputModule.h
+++ b/Engine/Plugins/Media/PixelStreaming2/Source/PixelStreaming2Input/Public/IPixelStreaming2InputModule.h
@@ -2,7 +2,6 @@
 
 #pragma once
 
-#include "Modules/ModuleManager.h"
 #include "IInputDeviceModule.h"
 #include "IPixelStreaming2InputHandler.h"
 
diff --git a/Engine/Plugins/Media/PixelStreaming2/Source/PixelStreaming2RTC/Internal/EpicRtcVideoCommon.h b/Engine/Plugins/Media/PixelStreaming2/Source/PixelStreaming2RTC/Internal/EpicRtcVideoCommon.h
--- a/Engine/Plugins/Media/PixelStreaming2/Source/PixelStreaming2RTC/Internal/EpicRtcVideoCommon.h
+++ b/Engine/Plugins/Media/PixelStreaming2/Source/PixelStreaming2RTC/Internal/EpicRtcVideoCommon.h
@@ -3,12 +3,9 @@
 #pragma once
 
 #include "Containers/Utf8String.h"
-#include "Video/DependencyDescriptor.h"
 #include "Video/GenericFrameInfo.h"
 #include "Video/VideoEncoder.h"
 
-#include "epic_rtc/containers/epic_rtc_array.h"
-#include "epic_rtc/containers/epic_rtc_string_view.h"
 #include "epic_rtc/core/video/video_buffer.h"
 #include "epic_rtc/core/video/video_codec_info.h"
 #include "epic_rtc_helper/memory/ref_count_impl_helper.h"
@@ -985,4 +982,4 @@ namespace UE::PixelStreaming2
 		EPICRTC_REFCOUNT_INTERFACE_IN_PLACE
 		// End EpicRtcRefCountInterface
 	};
-} // namespace UE::PixelStreaming2
\ No newline at end of file
+} // namespace UE::PixelStreaming2
diff --git a/Engine/Plugins/Media/PixelStreaming2/Source/PixelStreaming2RTC/Internal/UtilsString.h b/Engine/Plugins/Media/PixelStreaming2/Source/PixelStreaming2RTC/Internal/UtilsString.h
--- a/Engine/Plugins/Media/PixelStreaming2/Source/PixelStreaming2RTC/Internal/UtilsString.h
+++ b/Engine/Plugins/Media/PixelStreaming2/Source/PixelStreaming2RTC/Internal/UtilsString.h
@@ -4,9 +4,7 @@
 
 #include <string>
 
-#include "Containers/StringConv.h"
 #include "Containers/UnrealString.h"
-#include "Containers/Utf8String.h"
 #include "Dom/JsonObject.h"
 #include "Policies/PrettyJsonPrintPolicy.h"
 #include "Policies/CondensedJsonPrintPolicy.h"
diff --git a/Engine/Plugins/Media/PixelStreaming2/Source/PixelStreaming2RTC/Private/DebugGraph.h b/Engine/Plugins/Media/PixelStreaming2/Source/PixelStreaming2RTC/Private/DebugGraph.h
--- a/Engine/Plugins/Media/PixelStreaming2/Source/PixelStreaming2RTC/Private/DebugGraph.h
+++ b/Engine/Plugins/Media/PixelStreaming2/Source/PixelStreaming2RTC/Private/DebugGraph.h
@@ -3,7 +3,7 @@
 #pragma once
 
 #include "CanvasTypes.h"
-#include "CoreMinimal.h"
+#include "Containers/Array.h"
 
 namespace UE::PixelStreaming2
 {
diff --git a/Engine/Plugins/Media/PixelStreaming2/Source/PixelStreaming2RTC/Private/EpicRtcAudioCapturer.cpp b/Engine/Plugins/Media/PixelStreaming2/Source/PixelStreaming2RTC/Private/EpicRtcAudioCapturer.cpp
--- a/Engine/Plugins/Media/PixelStreaming2/Source/PixelStreaming2RTC/Private/EpicRtcAudioCapturer.cpp
+++ b/Engine/Plugins/Media/PixelStreaming2/Source/PixelStreaming2RTC/Private/EpicRtcAudioCapturer.cpp
@@ -3,11 +3,8 @@
 #include "EpicRtcAudioCapturer.h"
 
 #include "AudioDevice.h"
-#include "Engine/Engine.h"
-#include "Logging.h"
 #include "Misc/CoreDelegates.h"
 #include "PixelStreaming2PluginSettings.h"
-#include "Sound/SampleBufferIO.h"
 
 namespace UE::PixelStreaming2
 {
diff --git a/Engine/Plugins/Media/PixelStreaming2/Source/PixelStreaming2RTC/Private/EpicRtcAudioSource.cpp b/Engine/Plugins/Media/PixelStreaming2/Source/PixelStreaming2RTC/Private/EpicRtcAudioSource.cpp
--- a/Engine/Plugins/Media/PixelStreaming2/Source/PixelStreaming2RTC/Private/EpicRtcAudioSource.cpp
+++ b/Engine/Plugins/Media/PixelStreaming2/Source/PixelStreaming2RTC/Private/EpicRtcAudioSource.cpp
@@ -3,7 +3,6 @@
 #include "EpicRtcAudioSource.h"
 
 #include "EpicRtcAudioCapturer.h"
-#include "PixelStreaming2RTCModule.h"
 
 namespace UE::PixelStreaming2
 {
diff --git a/Engine/Plugins/Media/PixelStreaming2/Source/PixelStreaming2RTC/Private/EpicRtcDataTrack.cpp b/Engine/Plugins/Media/PixelStreaming2/Source/PixelStreaming2RTC/Private/EpicRtcDataTrack.cpp
--- a/Engine/Plugins/Media/PixelStreaming2/Source/PixelStreaming2RTC/Private/EpicRtcDataTrack.cpp
+++ b/Engine/Plugins/Media/PixelStreaming2/Source/PixelStreaming2RTC/Private/EpicRtcDataTrack.cpp
@@ -3,8 +3,6 @@
 #include "EpicRtcDataTrack.h"
 
 #include "Logging.h"
-#include "UtilsCommon.h"
-#include "UtilsString.h"
 
 namespace UE::PixelStreaming2
 {
diff --git a/Engine/Plugins/Media/PixelStreaming2/Source/PixelStreaming2RTC/Private/EpicRtcStreamer.cpp b/Engine/Plugins/Media/PixelStreaming2/Source/PixelStreaming2RTC/Private/EpicRtcStreamer.cpp
--- a/Engine/Plugins/Media/PixelStreaming2/Source/PixelStreaming2RTC/Private/EpicRtcStreamer.cpp
+++ b/Engine/Plugins/Media/PixelStreaming2/Source/PixelStreaming2RTC/Private/EpicRtcStreamer.cpp
@@ -6,34 +6,15 @@
 #include "PixelStreaming2Delegates.h"
 #include "Engine/GameEngine.h"
 #include "Engine/GameViewportClient.h"
-#include "EpicRtcWebsocket.h"
-#include "EpicRtcWebsocketFactory.h"
-#include "TextureResource.h"
-#include "WebSocketsModule.h"
-#include "Logging.h"
 #include "Stats.h"
 #include "PixelStreaming2StatNames.h"
-#include "Async/Async.h"
-#include "Engine/Texture2D.h"
-#include "RTCStatsCollector.h"
-#include "Framework/Application/SlateApplication.h"
-#include "Dom/JsonValue.h"
-#include "Serialization/JsonWriter.h"
-#include "Serialization/JsonSerializer.h"
-#include "Dom/JsonObject.h"
 #include "PixelStreaming2RTCModule.h"
 #include "EpicRtcTrack.h"
-#include "EpicRtcDataTrack.h"
-#include "IPixelStreaming2InputModule.h"
-#include "IPixelStreaming2Module.h"
-#include "CoreGlobals.h"
 #include "UtilsAsync.h"
-#include "UtilsString.h"
 #include "UtilsCodecs.h"
 #include "VideoProducerBackBuffer.h"
 #include "RTCInputHandler.h"
 #include "PixelStreaming2Common.h"
-#include "EpicRtcAudioCapturer.h"
 
 namespace UE::PixelStreaming2
 {
diff --git a/Engine/Plugins/Media/PixelStreaming2/Source/PixelStreaming2RTC/Private/EpicRtcStreamer.h b/Engine/Plugins/Media/PixelStreaming2/Source/PixelStreaming2RTC/Private/EpicRtcStreamer.h
--- a/Engine/Plugins/Media/PixelStreaming2/Source/PixelStreaming2RTC/Private/EpicRtcStreamer.h
+++ b/Engine/Plugins/Media/PixelStreaming2/Source/PixelStreaming2RTC/Private/EpicRtcStreamer.h
@@ -2,7 +2,6 @@
 
 #pragma once
 
-#include "Dom/JsonObject.h"
 #include "EpicRtcAudioTrackObserver.h"
 #include "EpicRtcAudioTrackObserverFactory.h"
 #include "EpicRtcDataTrackObserver.h"
@@ -12,13 +11,11 @@
 #include "EpicRtcRoomObserver.h"
 #include "EpicRtcSessionObserver.h"
 #include "FreezeFrame.h"
-#include "IPixelStreaming2InputHandler.h"
 #include "IPixelStreaming2Streamer.h"
 #include "Logging.h"
 #include "PlayerContext.h"
 #include "SharedTickableTasks.h"
 #include "StreamerReconnectTimer.h"
-#include "Templates/SharedPointer.h"
 #include "ThreadSafeMap.h"
 #include "UtilsString.h"
 #include "VideoSourceGroup.h"
diff --git a/Engine/Plugins/Media/PixelStreaming2/Source/PixelStreaming2RTC/Private/EpicRtcVideoDecoder.cpp b/Engine/Plugins/Media/PixelStreaming2/Source/PixelStreaming2RTC/Private/EpicRtcVideoDecoder.cpp
--- a/Engine/Plugins/Media/PixelStreaming2/Source/PixelStreaming2RTC/Private/EpicRtcVideoDecoder.cpp
+++ b/Engine/Plugins/Media/PixelStreaming2/Source/PixelStreaming2RTC/Private/EpicRtcVideoDecoder.cpp
@@ -5,7 +5,6 @@
 #include "ColorConversion.h"
 #include "EpicRtcVideoBufferI420.h"
 #include "EpicRtcVideoBufferRHI.h"
-#include "Logging.h"
 #include "PixelStreaming2Trace.h"
 #include "UtilsString.h"
 #include "UtilsCodecs.h"
diff --git a/Engine/Plugins/Media/PixelStreaming2/Source/PixelStreaming2RTC/Private/EpicRtcVideoDecoder.h b/Engine/Plugins/Media/PixelStreaming2/Source/PixelStreaming2RTC/Private/EpicRtcVideoDecoder.h
--- a/Engine/Plugins/Media/PixelStreaming2/Source/PixelStreaming2RTC/Private/EpicRtcVideoDecoder.h
+++ b/Engine/Plugins/Media/PixelStreaming2/Source/PixelStreaming2RTC/Private/EpicRtcVideoDecoder.h
@@ -9,7 +9,6 @@
 #include "Video/Decoders/Configs/VideoDecoderConfigVP9.h"
 #include "Video/Resources/VideoResourceCPU.h"
 #include "Video/Resources/VideoResourceRHI.h"
-#include "Video/VideoDecoder.h"
 
 #include "epic_rtc/core/video/video_decoder.h"
 #include "epic_rtc_helper/memory/ref_count_impl_helper.h"
diff --git a/Engine/Plugins/Media/PixelStreaming2/Source/PixelStreaming2RTC/Private/EpicRtcVideoDecoderInitializer.cpp b/Engine/Plugins/Media/PixelStreaming2/Source/PixelStreaming2RTC/Private/EpicRtcVideoDecoderInitializer.cpp
--- a/Engine/Plugins/Media/PixelStreaming2/Source/PixelStreaming2RTC/Private/EpicRtcVideoDecoderInitializer.cpp
+++ b/Engine/Plugins/Media/PixelStreaming2/Source/PixelStreaming2RTC/Private/EpicRtcVideoDecoderInitializer.cpp
@@ -4,17 +4,13 @@
 
 #include "EpicRtcVideoDecoder.h"
 #include "UtilsCoder.h"
-#include "Logging.h"
 #include "PixelStreaming2PluginSettings.h"
-#include "UtilsCommon.h"
 #include "UtilsString.h"
 #include "UtilsCodecs.h"
 #include "Video/Decoders/Configs/VideoDecoderConfigAV1.h"
 #include "Video/Decoders/Configs/VideoDecoderConfigH264.h"
 #include "Video/Decoders/Configs/VideoDecoderConfigVP8.h"
 #include "Video/Decoders/Configs/VideoDecoderConfigVP9.h"
-#include "Video/Resources/VideoResourceCPU.h"
-#include "Video/Resources/VideoResourceRHI.h"
 
 namespace
 {
diff --git a/Engine/Plugins/Media/PixelStreaming2/Source/PixelStreaming2RTC/Private/EpicRtcVideoEncoderInitializer.cpp b/Engine/Plugins/Media/PixelStreaming2/Source/PixelStreaming2RTC/Private/EpicRtcVideoEncoderInitializer.cpp
--- a/Engine/Plugins/Media/PixelStreaming2/Source/PixelStreaming2RTC/Private/EpicRtcVideoEncoderInitializer.cpp
+++ b/Engine/Plugins/Media/PixelStreaming2/Source/PixelStreaming2RTC/Private/EpicRtcVideoEncoderInitializer.cpp
@@ -9,7 +9,6 @@
 #include "PixelStreaming2Delegates.h"
 #include "PixelStreaming2PluginSettings.h"
 #include "UtilsAsync.h"
-#include "UtilsCommon.h"
 #include "UtilsString.h"
 #include "UtilsCodecs.h"
 #include "Video/Encoders/Configs/VideoEncoderConfigAV1.h"
diff --git a/Engine/Plugins/Media/PixelStreaming2/Source/PixelStreaming2RTC/Private/EpicRtcWebsocket.cpp b/Engine/Plugins/Media/PixelStreaming2/Source/PixelStreaming2RTC/Private/EpicRtcWebsocket.cpp
--- a/Engine/Plugins/Media/PixelStreaming2/Source/PixelStreaming2RTC/Private/EpicRtcWebsocket.cpp
+++ b/Engine/Plugins/Media/PixelStreaming2/Source/PixelStreaming2RTC/Private/EpicRtcWebsocket.cpp
@@ -5,9 +5,7 @@
 #include "IPixelStreaming2Module.h"
 #include "IWebSocket.h"
 #include "PixelStreaming2PluginSettings.h"
-#include "Serialization/JsonSerializer.h"
 #include "UtilsString.h"
-#include "UtilsCodecs.h"
 #include "WebSocketsModule.h"
 
 namespace UE::PixelStreaming2
diff --git a/Engine/Plugins/Media/PixelStreaming2/Source/PixelStreaming2RTC/Private/Logging.h b/Engine/Plugins/Media/PixelStreaming2/Source/PixelStreaming2RTC/Private/Logging.h
--- a/Engine/Plugins/Media/PixelStreaming2/Source/PixelStreaming2RTC/Private/Logging.h
+++ b/Engine/Plugins/Media/PixelStreaming2/Source/PixelStreaming2RTC/Private/Logging.h
@@ -2,7 +2,6 @@
 
 #pragma once
 
-#include "HAL/IConsoleManager.h"
 #include "Logging/LogMacros.h"
 #include "Logging/StructuredLog.h"
 
diff --git a/Engine/Plugins/Media/PixelStreaming2/Source/PixelStreaming2RTC/Private/PlayerContext.h b/Engine/Plugins/Media/PixelStreaming2/Source/PixelStreaming2RTC/Private/PlayerContext.h
--- a/Engine/Plugins/Media/PixelStreaming2/Source/PixelStreaming2RTC/Private/PlayerContext.h
+++ b/Engine/Plugins/Media/PixelStreaming2/Source/PixelStreaming2RTC/Private/PlayerContext.h
@@ -8,7 +8,6 @@
 #include "EpicRtcVideoSink.h"
 #include "EpicRtcVideoSource.h"
 #include "RTCStatsCollector.h"
-#include "Templates/RefCounting.h"
 #include "UtilsCommon.h"
 
 #include "epic_rtc/core/audio/audio_track.h"
diff --git a/Engine/Plugins/Media/PixelStreaming2/Source/PixelStreaming2RTC/Private/Stats.h b/Engine/Plugins/Media/PixelStreaming2/Source/PixelStreaming2RTC/Private/Stats.h
--- a/Engine/Plugins/Media/PixelStreaming2/Source/PixelStreaming2RTC/Private/Stats.h
+++ b/Engine/Plugins/Media/PixelStreaming2/Source/PixelStreaming2RTC/Private/Stats.h
@@ -7,7 +7,6 @@
 #include "UnrealEngine.h"
 #include "ConsoleSettings.h"
 #include "DebugGraph.h"
-#include "Misc/Optional.h"
 #include "PixelCaptureFrameMetadata.h"
 #include "IPixelStreaming2Stats.h"
 #include "RTCStatsCollector.h"
diff --git a/Engine/Plugins/Media/PixelStreaming2/Source/PixelStreaming2RTC/Private/UtilsCodecs.h b/Engine/Plugins/Media/PixelStreaming2/Source/PixelStreaming2RTC/Private/UtilsCodecs.h
--- a/Engine/Plugins/Media/PixelStreaming2/Source/PixelStreaming2RTC/Private/UtilsCodecs.h
+++ b/Engine/Plugins/Media/PixelStreaming2/Source/PixelStreaming2RTC/Private/UtilsCodecs.h
@@ -2,13 +2,9 @@
 
 #pragma once
 
-#include "Containers/Array.h"
 #include "EpicRtcVideoCommon.h"
-#include "IPixelStreaming2Streamer.h"
 #include "Logging.h"
-#include "PixelStreaming2PluginSettings.h"
 #include "UtilsVideo.h"
-#include "Video/VideoEncoder.h"
 
 #include "epic_rtc/core/video/video_common.h"
 #include "epic_rtc/core/video/video_rate_control.h"
diff --git a/Engine/Plugins/Media/PixelStreaming2/Source/PixelStreaming2RTC/Private/UtilsCommon.h b/Engine/Plugins/Media/PixelStreaming2/Source/PixelStreaming2RTC/Private/UtilsCommon.h
--- a/Engine/Plugins/Media/PixelStreaming2/Source/PixelStreaming2RTC/Private/UtilsCommon.h
+++ b/Engine/Plugins/Media/PixelStreaming2/Source/PixelStreaming2RTC/Private/UtilsCommon.h
@@ -2,14 +2,7 @@
 
 #pragma once
 
-#include "Async/Async.h"
-#include "CoreMinimal.h"
-#include "EpicRtcVideoCommon.h"
-#include "HAL/IConsoleManager.h"
 #include "Misc/CommandLine.h"
-#include "Video/CodecUtils/CodecUtilsH264.h"
-#include "Video/CodecUtils/CodecUtilsVP9.h"
-
 namespace UE::PixelStreaming2
 {
 	static const FString SFU_PLAYER_ID_PREFIX = FString(TEXT("SFU_"));
diff --git a/Engine/Plugins/Media/PixelStreaming2/Source/PixelStreaming2RTC/Public/IPixelStreaming2RTCModule.h b/Engine/Plugins/Media/PixelStreaming2/Source/PixelStreaming2RTC/Public/IPixelStreaming2RTCModule.h
--- a/Engine/Plugins/Media/PixelStreaming2/Source/PixelStreaming2RTC/Public/IPixelStreaming2RTCModule.h
+++ b/Engine/Plugins/Media/PixelStreaming2/Source/PixelStreaming2RTC/Public/IPixelStreaming2RTCModule.h
@@ -2,15 +2,11 @@
 
 #pragma once
 
-#include "CoreMinimal.h"
 #include "Modules/ModuleInterface.h"
 #include "Modules/ModuleManager.h"
 #include "Templates/SharedPointer.h"
-#include "Templates/RefCounting.h"
-#include "IPixelStreaming2AudioProducer.h"
 #include "IPixelStreaming2Streamer.h"
-
-#include "Video/VideoConfig.h"
+#include "Framework/Application/SlateApplication.h"
 
 /**
  * The IPixelStreaming2RTCModule interface manages the core functionality of the Pixel Streaming system.
diff --git a/Engine/Plugins/Media/PixelStreaming2/Source/PixelStreaming2Servers/Private/Logging.h b/Engine/Plugins/Media/PixelStreaming2/Source/PixelStreaming2Servers/Private/Logging.h
--- a/Engine/Plugins/Media/PixelStreaming2/Source/PixelStreaming2Servers/Private/Logging.h
+++ b/Engine/Plugins/Media/PixelStreaming2/Source/PixelStreaming2Servers/Private/Logging.h
@@ -2,7 +2,6 @@
 
 #pragma once
 
-#include "HAL/IConsoleManager.h"
 #include "Logging/LogMacros.h"
 #include "Logging/StructuredLog.h"
 
diff --git a/Engine/Plugins/Media/PixelStreaming2/Source/PixelStreaming2Servers/Private/PixelStreaming2ServersModule.h b/Engine/Plugins/Media/PixelStreaming2/Source/PixelStreaming2Servers/Private/PixelStreaming2ServersModule.h
--- a/Engine/Plugins/Media/PixelStreaming2/Source/PixelStreaming2Servers/Private/PixelStreaming2ServersModule.h
+++ b/Engine/Plugins/Media/PixelStreaming2/Source/PixelStreaming2Servers/Private/PixelStreaming2ServersModule.h
@@ -2,8 +2,6 @@
 
 #pragma once
 
-#include "CoreMinimal.h"
-#include "Containers/UnrealString.h"
 #include "Modules/ModuleManager.h"
 
 class FPixelStreaming2ServersModule : public IModuleInterface
diff --git a/Engine/Plugins/Media/PixelStreaming2/Source/PixelStreaming2Servers/Private/ServerUtils.cpp b/Engine/Plugins/Media/PixelStreaming2/Source/PixelStreaming2Servers/Private/ServerUtils.cpp
--- a/Engine/Plugins/Media/PixelStreaming2/Source/PixelStreaming2Servers/Private/ServerUtils.cpp
+++ b/Engine/Plugins/Media/PixelStreaming2/Source/PixelStreaming2Servers/Private/ServerUtils.cpp
@@ -4,7 +4,6 @@
 #include "HAL/PlatformFileManager.h"
 #include "PixelStreaming2Servers.h"
 #include "PixelStreaming2ServersModule.h"
-#include "Containers/UnrealString.h"
 #include "Logging.h"
 #include "Serialization/JsonSerializer.h"
 #include "Serialization/JsonWriter.h"
diff --git a/Engine/Plugins/Media/PixelStreaming2/Source/PixelStreaming2Servers/Private/ServerUtils.h b/Engine/Plugins/Media/PixelStreaming2/Source/PixelStreaming2Servers/Private/ServerUtils.h
--- a/Engine/Plugins/Media/PixelStreaming2/Source/PixelStreaming2Servers/Private/ServerUtils.h
+++ b/Engine/Plugins/Media/PixelStreaming2/Source/PixelStreaming2Servers/Private/ServerUtils.h
@@ -4,9 +4,7 @@
 
 #include "PixelStreaming2Servers.h"
 #include "Templates/SharedPointer.h"
-#include "Misc/MonitoredProcess.h"
 #include "Dom/JsonObject.h"
-#include "Engine/EngineBaseTypes.h"
 
 namespace UE::PixelStreaming2Servers
 {
diff --git a/Engine/Plugins/Media/PixelStreaming2/Source/PixelStreaming2Servers/Private/SignallingServer.h b/Engine/Plugins/Media/PixelStreaming2/Source/PixelStreaming2Servers/Private/SignallingServer.h
--- a/Engine/Plugins/Media/PixelStreaming2/Source/PixelStreaming2Servers/Private/SignallingServer.h
+++ b/Engine/Plugins/Media/PixelStreaming2/Source/PixelStreaming2Servers/Private/SignallingServer.h
@@ -5,9 +5,6 @@
 #include "ServerBase.h"
 #include "WebSocketProbe.h"
 #include "WebSocketServerWrapper.h"
-#include "Templates/UniquePtr.h"
-#include "Templates/SharedPointer.h"
-#include "Dom/JsonObject.h"
 
 namespace UE::PixelStreaming2Servers
 {
diff --git a/Engine/Plugins/Media/PixelStreaming2/Source/PixelStreaming2Servers/Private/WebSocketProbe.h b/Engine/Plugins/Media/PixelStreaming2/Source/PixelStreaming2Servers/Private/WebSocketProbe.h
--- a/Engine/Plugins/Media/PixelStreaming2/Source/PixelStreaming2Servers/Private/WebSocketProbe.h
+++ b/Engine/Plugins/Media/PixelStreaming2/Source/PixelStreaming2Servers/Private/WebSocketProbe.h
@@ -3,11 +3,9 @@
 #pragma once
 
 #include "IWebSocket.h"
-#include "HAL/ThreadSafeCounter.h"
 #include "WebSocketsModule.h"
 #include "Templates/SharedPointer.h"
 #include "Logging.h"
-#include "Engine/EngineBaseTypes.h"
 #include "ServerUtils.h"
 
 namespace UE::PixelStreaming2Servers
diff --git a/Engine/Plugins/Media/PixelStreaming2/Source/PixelStreaming2Servers/Private/WebSocketServerWrapper.h b/Engine/Plugins/Media/PixelStreaming2/Source/PixelStreaming2Servers/Private/WebSocketServerWrapper.h
--- a/Engine/Plugins/Media/PixelStreaming2/Source/PixelStreaming2Servers/Private/WebSocketServerWrapper.h
+++ b/Engine/Plugins/Media/PixelStreaming2/Source/PixelStreaming2Servers/Private/WebSocketServerWrapper.h
@@ -3,13 +3,8 @@
 #pragma once
 
 #include "IWebSocketServer.h"
-#include "Templates/UniquePtr.h"
-#include "Templates/SharedPointer.h"
-#include "WebSocketNetworkingDelegates.h"
 #include "INetworkingWebSocket.h"
 #include "Tickable.h"
-#include "HAL/ThreadSafeCounter.h"
-#include "Containers/ArrayView.h"
 #include "HAL/ThreadSafeBool.h"
 
 namespace UE::PixelStreaming2Servers
diff --git a/Engine/Plugins/Media/PixelStreaming2/Source/PixelStreaming2Settings/Internal/PixelStreaming2PluginSettings.h b/Engine/Plugins/Media/PixelStreaming2/Source/PixelStreaming2Settings/Internal/PixelStreaming2PluginSettings.h
--- a/Engine/Plugins/Media/PixelStreaming2/Source/PixelStreaming2Settings/Internal/PixelStreaming2PluginSettings.h
+++ b/Engine/Plugins/Media/PixelStreaming2/Source/PixelStreaming2Settings/Internal/PixelStreaming2PluginSettings.h
@@ -5,7 +5,6 @@
 #include "Containers/UnrealString.h"
 #include "CoreMinimal.h"
 #include "Engine/DeveloperSettings.h"
-#include "HAL/IConsoleManager.h"
 #include "PixelStreaming2SettingsEnums.h"
 
 #include "AVConfig.h"
diff --git a/Engine/Plugins/Media/PixelStreaming2/Source/PixelStreaming2Settings/Private/Logging.h b/Engine/Plugins/Media/PixelStreaming2/Source/PixelStreaming2Settings/Private/Logging.h
--- a/Engine/Plugins/Media/PixelStreaming2/Source/PixelStreaming2Settings/Private/Logging.h
+++ b/Engine/Plugins/Media/PixelStreaming2/Source/PixelStreaming2Settings/Private/Logging.h
@@ -2,7 +2,6 @@
 
 #pragma once
 
-#include "HAL/IConsoleManager.h"
 #include "Logging/LogMacros.h"
 #include "Logging/StructuredLog.h"
 
diff --git a/Engine/Plugins/Media/PixelStreaming2/Source/PixelStreaming2Settings/Public/PixelStreaming2SettingsEnums.h b/Engine/Plugins/Media/PixelStreaming2/Source/PixelStreaming2Settings/Public/PixelStreaming2SettingsEnums.h
--- a/Engine/Plugins/Media/PixelStreaming2/Source/PixelStreaming2Settings/Public/PixelStreaming2SettingsEnums.h
+++ b/Engine/Plugins/Media/PixelStreaming2/Source/PixelStreaming2Settings/Public/PixelStreaming2SettingsEnums.h
@@ -2,8 +2,6 @@
 
 #pragma once
 
-#include "CoreMinimal.h"
-
 #include "PixelStreaming2SettingsEnums.generated.h"
 
 /**
