[PixelStreaming2] Fix: Crash during starting of Play In Editor session due to nullptr.

diff --git a/Engine/Plugins/Media/PixelStreaming2/Source/PixelStreaming2Editor/Private/PixelStreaming2EditorModule.cpp b/Engine/Plugins/Media/PixelStreaming2/Source/PixelStreaming2Editor/Private/PixelStreaming2EditorModule.cpp
--- a/Engine/Plugins/Media/PixelStreaming2/Source/PixelStreaming2Editor/Private/PixelStreaming2EditorModule.cpp
+++ b/Engine/Plugins/Media/PixelStreaming2/Source/PixelStreaming2Editor/Private/PixelStreaming2EditorModule.cpp
@@ -587,6 +587,11 @@ namespace UE::EditorPixelStreaming2
 
 	void FPixelStreaming2EditorModule::OnEndPIE(bool bIsSimulating)
 	{
+		if (!PIEStreamer)
+		{
+			return;
+		}
+		
 		PIEStreamer->StopStreaming();
 
 		IPixelStreaming2Module::Get().DeleteStreamer(PIEStreamer);
