[PixelStreaming2] Fix: Poor bitrate during subsequent peer connections

diff --git a/Engine/Plugins/Media/PixelStreaming2/Source/PixelStreaming2RTC/Private/EpicRtcStreamer.cpp b/Engine/Plugins/Media/PixelStreaming2/Source/PixelStreaming2RTC/Private/EpicRtcStreamer.cpp
--- a/Engine/Plugins/Media/PixelStreaming2/Source/PixelStreaming2RTC/Private/EpicRtcStreamer.cpp
+++ b/Engine/Plugins/Media/PixelStreaming2/Source/PixelStreaming2RTC/Private/EpicRtcStreamer.cpp
@@ -333,21 +333,27 @@ namespace UE::PixelStreaming2
 		SendAllPlayersMessage(MessageName, Descriptor);
 	}
 
-	EpicRtcBitrate FEpicRtcStreamer::GetBitrates() const
+	EpicRtcBitrate FEpicRtcStreamer::GetBitrates(bool bIncludeStartBitrate) const
 	{
 		const int MinBitrate = UPixelStreaming2PluginSettings::CVarWebRTCMinBitrate.GetValueOnAnyThread();
 		const int MaxBitrate = UPixelStreaming2PluginSettings::CVarWebRTCMaxBitrate.GetValueOnAnyThread();
-		const int StartBitrate = FMath::Clamp(UPixelStreaming2PluginSettings::CVarWebRTCStartBitrate.GetValueOnAnyThread(), MinBitrate, MaxBitrate);
 
 		EpicRtcBitrate Bitrates = {
 			._minBitrateBps = MinBitrate,
 			._hasMinBitrateBps = true,
 			._maxBitrateBps = MaxBitrate,
 			._hasMaxBitrateBps = true,
-			._startBitrateBps = StartBitrate,
-			._hasStartBitrateBps = true
+			._startBitrateBps = 0,
+			._hasStartBitrateBps = false
 		};
 
+		if (bIncludeStartBitrate)
+		{
+			const int StartBitrate = FMath::Clamp(UPixelStreaming2PluginSettings::CVarWebRTCStartBitrate.GetValueOnAnyThread(), MinBitrate, MaxBitrate);
+			Bitrates._startBitrateBps = StartBitrate;
+			Bitrates._hasStartBitrateBps = true;
+		}
+
 		return Bitrates;
 	}
 
@@ -556,7 +562,9 @@ namespace UE::PixelStreaming2
 				return;
 			}
 
-			ConnectionInterface->SetConnectionRates(GetBitrates());
+			// We don't include the starting bitrate during a bitrate refresh as that causes the peer connection
+			// to go back to that value and have to work its way up. This function is also only triggered by changes to min and max
+			ConnectionInterface->SetConnectionRates(GetBitrates(false /* bIncludeStartBitrate */));
 		});
 	}
 
@@ -1094,7 +1102,7 @@ namespace UE::PixelStreaming2
 					}
 				});
 
-				const EpicRtcBitrate Bitrate = GetBitrates();
+				const EpicRtcBitrate Bitrate = GetBitrates(true /* bIncludeStartBitrate */);
 
 				EpicRtcPortAllocator PortAllocator = {
 					._minPort = UPixelStreaming2PluginSettings::CVarWebRTCMinPort.GetValueOnAnyThread(),
@@ -1505,7 +1513,7 @@ namespace UE::PixelStreaming2
 			ParticipantConnection->AddDataSource(DataSource);
 		}
 
-		ParticipantConnection->SetConnectionRates(GetBitrates());
+		ParticipantConnection->SetConnectionRates(GetBitrates(true /* bIncludeStartBitrate */));
 
 		ParticipantConnection->StartNegotiation();
 	}
diff --git a/Engine/Plugins/Media/PixelStreaming2/Source/PixelStreaming2RTC/Private/EpicRtcStreamer.h b/Engine/Plugins/Media/PixelStreaming2/Source/PixelStreaming2RTC/Private/EpicRtcStreamer.h
--- a/Engine/Plugins/Media/PixelStreaming2/Source/PixelStreaming2RTC/Private/EpicRtcStreamer.h
+++ b/Engine/Plugins/Media/PixelStreaming2/Source/PixelStreaming2RTC/Private/EpicRtcStreamer.h
@@ -120,7 +120,7 @@ namespace UE::PixelStreaming2
 		void		   OnFrameCapturerCreated();
 		void		   OnUIInteraction(FMemoryReader Ar);
 		void		   OnSendMessage(FString MessageName, FMemoryReader Ar);
-		EpicRtcBitrate GetBitrates() const;
+		EpicRtcBitrate GetBitrates(bool bIncludeStartBitrate) const;
 
 		template <typename T>
 		bool FindPlayerFromTrack(T Track, FString& OutPlayerId)
