[PixelStreaming2] QOL: Add insights tracking to tickable tasks

diff --git a/Engine/Plugins/Media/PixelStreaming2/Source/PixelStreaming2/Private/Thread.cpp b/Engine/Plugins/Media/PixelStreaming2/Source/PixelStreaming2/Private/Thread.cpp
--- a/Engine/Plugins/Media/PixelStreaming2/Source/PixelStreaming2/Private/Thread.cpp
+++ b/Engine/Plugins/Media/PixelStreaming2/Source/PixelStreaming2/Private/Thread.cpp
@@ -8,6 +8,7 @@
 #include "HAL/PlatformTime.h"
 #include "HAL/RunnableThread.h"
 #include "Misc/SingleThreadRunnable.h"
+#include "PixelStreaming2Trace.h"
 #include "TickableTask.h"
 
 namespace UE::PixelStreaming2
@@ -32,6 +33,7 @@ namespace UE::PixelStreaming2
 
 			while (bIsRunning)
 			{
+				TRACE_CPUPROFILER_EVENT_SCOPE_ON_CHANNEL_STR("FPixelStreamingRunnable::Run LoopStart", PixelStreaming2Channel);
 				Tick();
 
 				// Sleep 1ms
@@ -43,6 +45,7 @@ namespace UE::PixelStreaming2
 
 		virtual void Stop() override
 		{
+			TRACE_CPUPROFILER_EVENT_SCOPE_ON_CHANNEL_STR("FPixelStreamingRunnable::Stop", PixelStreaming2Channel);
 			bIsRunning = false;
 
 			TaskEvent->Trigger();
@@ -50,6 +53,7 @@ namespace UE::PixelStreaming2
 
 		virtual void Exit() override
 		{
+			TRACE_CPUPROFILER_EVENT_SCOPE_ON_CHANNEL_STR("FPixelStreamingRunnable::Exit", PixelStreaming2Channel);
 			bIsRunning = false;
 
 			TaskEvent->Trigger();
@@ -57,6 +61,7 @@ namespace UE::PixelStreaming2
 
 		virtual FSingleThreadRunnable* GetSingleThreadInterface() override
 		{
+			TRACE_CPUPROFILER_EVENT_SCOPE_ON_CHANNEL_STR("FPixelStreamingRunnable::GetSingleThreadInterface", PixelStreaming2Channel);
 			return this;
 		}
 		// End FRunnable
@@ -64,6 +69,7 @@ namespace UE::PixelStreaming2
 		// Begin FSingleThreadRunnable
 		virtual void Tick() override
 		{
+			TRACE_CPUPROFILER_EVENT_SCOPE_ON_CHANNEL_STR("FPixelStreamingRunnable::Tick", PixelStreaming2Channel);
 			FScopeLock TaskLock(&TasksMutex);
 
 			const uint64 NowCycles = FPlatformTime::Cycles64();
@@ -77,6 +83,7 @@ namespace UE::PixelStreaming2
 				TSharedPtr<FPixelStreamingTickableTask> PinnedTask = Task.Pin();
 				if (PinnedTask)
 				{
+					TRACE_CPUPROFILER_EVENT_SCOPE_TEXT_ON_CHANNEL(*PinnedTask->GetName(), PixelStreaming2Channel)
 					PinnedTask->Tick(DeltaMs);
 				}
 			}
@@ -89,6 +96,7 @@ namespace UE::PixelStreaming2
 
 		void StartTicking()
 		{
+			TRACE_CPUPROFILER_EVENT_SCOPE_ON_CHANNEL_STR("FPixelStreamingRunnable::StartTicking", PixelStreaming2Channel);
 			FScopeLock NewTaskLock(&NewTasksMutex);
 			for (auto& NewTask : NewTasks)
 			{
@@ -102,6 +110,7 @@ namespace UE::PixelStreaming2
 
 		void FinishTicking()
 		{
+			TRACE_CPUPROFILER_EVENT_SCOPE_ON_CHANNEL_STR("FPixelStreamingRunnable::FinishTicking", PixelStreaming2Channel);
 			if (bNeedsCleanup)
 			{
 				Tasks.RemoveAll([](const TWeakPtr<FPixelStreamingTickableTask> Entry) { return Entry == nullptr || !Entry.IsValid(); });
@@ -131,6 +140,7 @@ namespace UE::PixelStreaming2
 	private:
 		void AddTask(TWeakPtr<FPixelStreamingTickableTask> Task)
 		{
+			TRACE_CPUPROFILER_EVENT_SCOPE_ON_CHANNEL_STR("FPixelStreamingRunnable::AddTask", PixelStreaming2Channel);
 			FScopeLock NewTaskLock(&NewTasksMutex);
 			NewTasks.Add(Task);
 			// We've added a new task. Wake the thread (if it was sleeping)
@@ -139,6 +149,7 @@ namespace UE::PixelStreaming2
 
 		void RemoveTask(FPixelStreamingTickableTask* Task)
 		{
+			TRACE_CPUPROFILER_EVENT_SCOPE_ON_CHANNEL_STR("FPixelStreamingRunnable::RemoveTask", PixelStreaming2Channel);
 			if (Task == nullptr)
 			{
 				return;
@@ -204,6 +215,7 @@ namespace UE::PixelStreaming2
 
 	FPixelStreamingThread::FPixelStreamingThread()
 	{
+		TRACE_CPUPROFILER_EVENT_SCOPE_ON_CHANNEL_STR("FPixelStreamingThread Constructor", PixelStreaming2Channel)
 		Runnable = MakeShared<FPixelStreamingRunnable>();
 		PixelStreamingRunnable = Runnable;
 		Thread = TSharedPtr<FRunnableThread>(FRunnableThread::Create(Runnable.Get(), TEXT("Pixel Streaming PixelStreaming Thread")));
@@ -211,6 +223,7 @@ namespace UE::PixelStreaming2
 
 	FPixelStreamingThread::~FPixelStreamingThread()
 	{
+		TRACE_CPUPROFILER_EVENT_SCOPE_ON_CHANNEL_STR("FPixelStreamingThread Destructor", PixelStreaming2Channel);
 		if (Thread)
 		{
 			Thread->Kill();
@@ -230,6 +243,7 @@ namespace UE::PixelStreaming2
 
 	void FPixelStreamingTickableTask::Register()
 	{
+		TRACE_CPUPROFILER_EVENT_SCOPE_ON_CHANNEL_STR("FPixelStreamingTickableTask::Register", PixelStreaming2Channel);
 		TSharedPtr<FPixelStreamingRunnable> Runnable = PixelStreamingRunnable.Pin();
 		if (Runnable)
 		{
@@ -239,6 +253,7 @@ namespace UE::PixelStreaming2
 
 	void FPixelStreamingTickableTask::Unregister()
 	{
+		TRACE_CPUPROFILER_EVENT_SCOPE_ON_CHANNEL_STR("FPixelStreamingTickableTask::Unregister", PixelStreaming2Channel);
 		TSharedPtr<FPixelStreamingRunnable> Runnable = PixelStreamingRunnable.Pin();
 		if (Runnable)
 		{
diff --git a/Engine/Plugins/Media/PixelStreaming2/Source/PixelStreaming2RTC/Internal/EpicRtcConferenceUtils.h b/Engine/Plugins/Media/PixelStreaming2/Source/PixelStreaming2RTC/Internal/EpicRtcConferenceUtils.h
--- a/Engine/Plugins/Media/PixelStreaming2/Source/PixelStreaming2RTC/Internal/EpicRtcConferenceUtils.h
+++ b/Engine/Plugins/Media/PixelStreaming2/Source/PixelStreaming2RTC/Internal/EpicRtcConferenceUtils.h
@@ -4,6 +4,7 @@
 
 #include "IPixelStreaming2Stats.h"
 #include "PixelStreaming2PluginSettings.h"
+#include "PixelStreaming2Trace.h"
 #include "TickableTask.h"
 
 #include "epic_rtc/core/conference.h"
@@ -39,9 +40,14 @@ namespace UE::PixelStreaming2
 			if (EpicRtcConference)
 			{
 				// Tick conference normally. This handles things like data channel message
-				IPixelStreaming2Stats::Get().GraphValue(TEXT("ConferenceTickInterval"), DeltaMs, 1, 0.f, 1.f);
+				{
+					TRACE_CPUPROFILER_EVENT_SCOPE_ON_CHANNEL_STR("FEpicRtcTickConferenceTask::Tick GraphValue", PixelStreaming2Channel);
+					IPixelStreaming2Stats::Get().GraphValue(TEXT("ConferenceTickInterval"), DeltaMs, 1, 0.f, 1.f);
+				}
+
 				while (EpicRtcConference->NeedsTick())
 				{
+					TRACE_CPUPROFILER_EVENT_SCOPE_ON_CHANNEL_STR("EpicRtcConferenceInterface::Tick", PixelStreaming2Channel);
 					EpicRtcConference->Tick();
 				}
 			}
diff --git a/Engine/Plugins/Media/PixelStreaming2/Source/PixelStreaming2RTC/Private/EpicRtcVideoEncoder.cpp b/Engine/Plugins/Media/PixelStreaming2/Source/PixelStreaming2RTC/Private/EpicRtcVideoEncoder.cpp
--- a/Engine/Plugins/Media/PixelStreaming2/Source/PixelStreaming2RTC/Private/EpicRtcVideoEncoder.cpp
+++ b/Engine/Plugins/Media/PixelStreaming2/Source/PixelStreaming2RTC/Private/EpicRtcVideoEncoder.cpp
@@ -11,6 +11,7 @@
 #include "PixelCaptureOutputFrameI420.h"
 #include "PixelCaptureOutputFrameRHI.h"
 #include "PixelStreaming2PluginSettings.h"
+#include "PixelStreaming2Trace.h"
 #include "Stats.h"
 #include "UtilsString.h"
 #include "UtilsCodecs.h"
@@ -277,6 +278,8 @@ namespace UE::PixelStreaming2
 	template <std::derived_from<FVideoResource> TVideoResource>
 	EpicRtcMediaResult TEpicRtcVideoEncoder<TVideoResource>::Encode(const EpicRtcVideoFrame& VideoFrame, EpicRtcVideoFrameTypeArrayInterface* FrameTypes)
 	{
+		TRACE_CPUPROFILER_EVENT_SCOPE_ON_CHANNEL_STR("TEpicRtcVideoEncoder::Encode", PixelStreaming2Channel);
+
 		// Capture the Callback to ensure it is not released in a different thread.
 		TRefCountPtr<EpicRtcVideoEncoderCallbackInterface> CallbackEncoded(VideoEncoderCallback);
 		if (!CallbackEncoded)
@@ -367,13 +370,18 @@ namespace UE::PixelStreaming2
 		const bool bKeyFrame = (FrameTypes && FrameTypes->Size() > 0 && (FrameTypes->Get()[0] == EpicRtcVideoFrameType::I));
 
 		// Encode
-		Encoder->SendFrame(VideoResource, VideoFrame._timestampUs, bKeyFrame);
+		{
+			TRACE_CPUPROFILER_EVENT_SCOPE_ON_CHANNEL_STR("TVideoEncoder::SendFrame", PixelStreaming2Channel);
+			Encoder->SendFrame(VideoResource, VideoFrame._timestampUs, bKeyFrame);
+		}
 
 		UpdateFrameMetadataPostEncode(*AdaptedLayer);
 
 		FVideoPacket Packet;
 		while (Encoder->ReceivePacket(Packet))
 		{
+			TRACE_CPUPROFILER_EVENT_SCOPE_ON_CHANNEL_STR("TVideoEncoder::ReceivePacket", PixelStreaming2Channel);
+			
 			TRefCountPtr<FEpicRtcEncodedVideoBuffer> EncodedBuffer = new FEpicRtcEncodedVideoBuffer(Packet.DataPtr.Get(), Packet.DataSize);
 
 			EpicRtcEncodedVideoFrame EncodedFrame = {
