RTCP-9050 Fix for OnDisconnect test fail

diff --git a/Engine/Plugins/Media/PixelStreaming2/Source/PixelStreaming2RTC/Private/EpicRtcStreamer.cpp b/Engine/Plugins/Media/PixelStreaming2/Source/PixelStreaming2RTC/Private/EpicRtcStreamer.cpp
--- a/Engine/Plugins/Media/PixelStreaming2/Source/PixelStreaming2RTC/Private/EpicRtcStreamer.cpp
+++ b/Engine/Plugins/Media/PixelStreaming2/Source/PixelStreaming2RTC/Private/EpicRtcStreamer.cpp
@@ -1147,7 +1147,7 @@ namespace UE::PixelStreaming2
 				UE_LOGFMT(LogPixelStreaming2RTC, VeryVerbose, "FEpicRtcStreamer::OnSessionStateUpdate State=Disconnected");
 
 				// If we were successfully connected previously then we should broadcast the delegates about disconnection
-				bool bBroadcastDelegates = StreamState == EStreamState::Connected;
+				bool bBroadcastDelegates = StreamState == EStreamState::Connected || StreamState == EStreamState::Disconnecting;
 
 				// If the session unexpectedly disconnects (ie signalling server goes away), we should try and reconnect
 				if (StreamState == EStreamState::Connected || StreamState == EStreamState::Connecting)
