[PS2] Fix: PIE and level editor viewport video wasn't capturing widgets

diff --git a/Engine/Plugins/Media/PixelStreaming2/Source/PixelStreaming2/Internal/VideoProducer.h b/Engine/Plugins/Media/PixelStreaming2/Source/PixelStreaming2/Internal/VideoProducer.h
--- a/Engine/Plugins/Media/PixelStreaming2/Source/PixelStreaming2/Internal/VideoProducer.h
+++ b/Engine/Plugins/Media/PixelStreaming2/Source/PixelStreaming2/Internal/VideoProducer.h
@@ -43,6 +43,8 @@ namespace UE::PixelStreaming2
 
 	protected:
 		FVideoProducer() = default;
+
+		virtual ETextureCreateFlags GetTexCreateFlags();
 	};
 
 } // namespace UE::PixelStreaming2
\ No newline at end of file
diff --git a/Engine/Plugins/Media/PixelStreaming2/Source/PixelStreaming2/Internal/VideoProducerPIEViewport.h b/Engine/Plugins/Media/PixelStreaming2/Source/PixelStreaming2/Internal/VideoProducerPIEViewport.h
--- a/Engine/Plugins/Media/PixelStreaming2/Source/PixelStreaming2/Internal/VideoProducerPIEViewport.h
+++ b/Engine/Plugins/Media/PixelStreaming2/Source/PixelStreaming2/Internal/VideoProducerPIEViewport.h
@@ -2,29 +2,28 @@
 
 #pragma once
 
-#include "VideoProducer.h"
-
-class FViewport;
+#include "RHIFwd.h"
+#include "VideoProducerViewportBase.h"
 
 namespace UE::PixelStreaming2
 {
 	/**
 	 * An extension of the back buffer input that can handle PIE sessions. Primarily to be used in blueprints
 	 */
-	class PIXELSTREAMING2_API FVideoProducerPIEViewport : public FVideoProducer
+	class PIXELSTREAMING2_API FVideoProducerPIEViewport : public FVideoProducerViewportBase
 	{
 	public:
 		static TSharedPtr<FVideoProducerPIEViewport> Create();
-		virtual ~FVideoProducerPIEViewport();
+		virtual ~FVideoProducerPIEViewport() = default;
 
 		virtual FString ToString() override;
 
+	protected:
+		virtual bool ShouldCaptureViewport() override;
+
 	private:
 		FVideoProducerPIEViewport() = default;
-
-		void OnViewportRendered(FViewport* InViewport);
-
-		FDelegateHandle DelegateHandle;
+ 
+		void OnPreTick(float DeltaTime);
 	};
-
 } // namespace UE::PixelStreaming2
\ No newline at end of file
diff --git a/Engine/Plugins/Media/PixelStreaming2/Source/PixelStreaming2/Internal/VideoProducerViewportBase.h b/Engine/Plugins/Media/PixelStreaming2/Source/PixelStreaming2/Internal/VideoProducerViewportBase.h
new file mode 100644
--- /dev/null
+++ b/Engine/Plugins/Media/PixelStreaming2/Source/PixelStreaming2/Internal/VideoProducerViewportBase.h
@@ -0,0 +1,41 @@
+// Copyright Epic Games, Inc. All Rights Reserved.
+
+#pragma once
+
+#include "RHIFwd.h"
+#include "VideoProducer.h"
+#include "Widgets/SViewport.h"
+#include "Widgets/SWindow.h"
+
+namespace UE::PixelStreaming2
+{
+	/**
+	 * A base video producer for streaming a target viewport by way of cropping the entire backbuffer to the size of the viewport widget. 
+	 * 
+	 * 
+	 * Implementations MUST set the TargetViewport and TargetWidget and then call CalculateCaptureRegion at least once to configure the captured region.
+	 * Additionally, implementations MUST override ShouldCaptureViewport with an optional check to prevent OnBackBufferReadyToPresent from sending frames.
+	 */
+	class PIXELSTREAMING2_API FVideoProducerViewportBase : public FVideoProducer
+	{
+	public:
+		virtual ~FVideoProducerViewportBase() = default;
+
+		virtual bool IsFrameAlreadyCopied() override;
+
+    protected:
+        FVideoProducerViewportBase() = default;
+
+        // Viewport specific implementations should override this method with their own check
+        virtual bool ShouldCaptureViewport();
+
+        void CalculateCaptureRegion(TSharedRef<SViewport> Viewport, TSharedRef<SWindow> Window);
+		void OnBackBufferReadyToPresent(SWindow& Window, const FTextureRHIRef& BackBuffer);
+
+        FIntRect CaptureRect;
+		FCriticalSection CaptureRectCS;
+
+        TWeakPtr<SViewport> TargetViewport;
+		TWeakPtr<SWindow> TargetWindow;
+	};
+} // namespace UE::PixelStreaming2
\ No newline at end of file
diff --git a/Engine/Plugins/Media/PixelStreaming2/Source/PixelStreaming2/PixelStreaming2.Build.cs b/Engine/Plugins/Media/PixelStreaming2/Source/PixelStreaming2/PixelStreaming2.Build.cs
--- a/Engine/Plugins/Media/PixelStreaming2/Source/PixelStreaming2/PixelStreaming2.Build.cs
+++ b/Engine/Plugins/Media/PixelStreaming2/Source/PixelStreaming2/PixelStreaming2.Build.cs
@@ -40,6 +40,15 @@ namespace UnrealBuildTool.Rules
                 "PixelStreaming2Settings"
             });
 
+            if (Target.bBuildEditor)
+			{
+				PrivateDependencyModuleNames.AddRange(new string[]
+				{
+                    // Needed for the PIE viewport video producer
+					"UnrealEd",
+				});
+			}
+
             if (Target.IsInPlatformGroup(UnrealPlatformGroup.Windows) || Target.IsInPlatformGroup(UnrealPlatformGroup.Linux))
             {
                 PrivateDependencyModuleNames.Add("VulkanRHI");
diff --git a/Engine/Plugins/Media/PixelStreaming2/Source/PixelStreaming2/Private/VideoProducer.cpp b/Engine/Plugins/Media/PixelStreaming2/Source/PixelStreaming2/Private/VideoProducer.cpp
--- a/Engine/Plugins/Media/PixelStreaming2/Source/PixelStreaming2/Private/VideoProducer.cpp
+++ b/Engine/Plugins/Media/PixelStreaming2/Source/PixelStreaming2/Private/VideoProducer.cpp
@@ -28,4 +28,19 @@ namespace UE::PixelStreaming2
 		return VideoProducerIdentifiers::FVideoProducer;
 	}
 
+	ETextureCreateFlags FVideoProducer::GetTexCreateFlags()
+	{
+		ETextureCreateFlags Flags = TexCreate_RenderTargetable | TexCreate_UAV;
+
+		if (RHIGetInterfaceType() == ERHIInterfaceType::Vulkan)
+		{
+			Flags |= TexCreate_External;
+		}
+		else if (RHIGetInterfaceType() == ERHIInterfaceType::D3D11 || RHIGetInterfaceType() == ERHIInterfaceType::D3D12)
+		{
+			Flags |= TexCreate_Shared;
+		}
+
+		return Flags;
+	}
 } // namespace UE::PixelStreaming2
\ No newline at end of file
diff --git a/Engine/Plugins/Media/PixelStreaming2/Source/PixelStreaming2/Private/VideoProducerPIEViewport.cpp b/Engine/Plugins/Media/PixelStreaming2/Source/PixelStreaming2/Private/VideoProducerPIEViewport.cpp
--- a/Engine/Plugins/Media/PixelStreaming2/Source/PixelStreaming2/Private/VideoProducerPIEViewport.cpp
+++ b/Engine/Plugins/Media/PixelStreaming2/Source/PixelStreaming2/Private/VideoProducerPIEViewport.cpp
@@ -2,15 +2,20 @@
 
 #include "VideoProducerPIEViewport.h"
 
-#include "Engine/GameViewportClient.h"
+#include "Framework/Application/SlateApplication.h"
+#include "Layout/ArrangedChildren.h"
+#include "Layout/Geometry.h"
+#include "Layout/WidgetPath.h"
 #include "PixelCaptureInputFrameRHI.h"
-#include "UtilsAsync.h"
 #include "Rendering/SlateRenderer.h"
-#include "UnrealClient.h"
+#include "UtilsAsync.h"
+
+#if WITH_EDITOR
+#include "Editor.h"
+#endif
 
 namespace UE::PixelStreaming2
 {
-
 	TSharedPtr<FVideoProducerPIEViewport> FVideoProducerPIEViewport::Create()
 	{
 		TSharedPtr<FVideoProducerPIEViewport> NewInput = TSharedPtr<FVideoProducerPIEViewport>(new FVideoProducerPIEViewport());
@@ -19,48 +24,49 @@ namespace UE::PixelStreaming2
 		UE::PixelStreaming2::DoOnGameThread([WeakInput]() {
 			if (TSharedPtr<FVideoProducerPIEViewport> Input = WeakInput.Pin())
 			{
-				Input->DelegateHandle = UGameViewportClient::OnViewportRendered().AddSP(Input.ToSharedRef(), &FVideoProducerPIEViewport::OnViewportRendered);
+				FSlateRenderer* Renderer = FSlateApplication::Get().GetRenderer();
+				if(!Renderer)
+				{
+					return;
+				}
+ 
+				Renderer->OnBackBufferReadyToPresent().AddSP(Input.ToSharedRef(), &FVideoProducerPIEViewport::OnBackBufferReadyToPresent);
+				FSlateApplication::Get().OnPreTick().AddSP(Input.ToSharedRef(), &FVideoProducerPIEViewport::OnPreTick);
 			}
 		});
 
 		return NewInput;
 	}
 
-	FVideoProducerPIEViewport::~FVideoProducerPIEViewport()
+	FString FVideoProducerPIEViewport::ToString()
 	{
-		if (!IsEngineExitRequested())
-		{
-			UGameViewportClient::OnViewportRendered().Remove(DelegateHandle);
-		}
+		return VideoProducerIdentifiers::FVideoProducerPIEViewport;
 	}
 
-	void FVideoProducerPIEViewport::OnViewportRendered(FViewport* InViewport)
+	bool FVideoProducerPIEViewport::ShouldCaptureViewport()
+	{
+#if WITH_EDITOR
+		return GEditor && GEditor->PlayWorld && GEditor->PlayWorld->WorldType == EWorldType::PIE;
+#else
+		return false;
+#endif
+	}
+	
+	void FVideoProducerPIEViewport::OnPreTick(float DeltaTime)
 	{
-		if (!InViewport->IsPlayInEditorViewport())
+		if (!ShouldCaptureViewport())
 		{
 			return;
 		}
-
-		const FTextureRHIRef& FrameBuffer = InViewport->GetRenderTargetTexture();
-		if (!FrameBuffer)
+	
+		TargetViewport = FSlateApplication::Get().GetGameViewport();
+		if (TargetViewport.IsValid())
 		{
-			return;
+			TargetWindow = FSlateApplication::Get().FindWidgetWindow(TargetViewport.Pin().ToSharedRef());
+			if (TargetWindow.IsValid())
+			{
+				CalculateCaptureRegion(TargetViewport.Pin().ToSharedRef(), TargetWindow.Pin().ToSharedRef());
+			}
 		}
-
-		TSharedPtr<FVideoProducerUserData> UserData = MakeShared<FVideoProducerUserData>();
-		UserData->ProducerName = TEXT("FVideoProducerPIEViewport");
-		UserData->ProductionBeginCycles = FPlatformTime::Cycles64();
-
-		ENQUEUE_RENDER_COMMAND(StreamViewportTextureCommand)
-		([&, FrameBuffer, UserData](FRHICommandList& RHICmdList) {
-			UserData->ProductionEndCycles = FPlatformTime::Cycles64();
-			PushFrame(FPixelCaptureInputFrameRHI(FrameBuffer, UserData));
-		});
 	}
-
-	FString FVideoProducerPIEViewport::ToString()
-	{
-		return VideoProducerIdentifiers::FVideoProducerPIEViewport;
-	}
-
 } // namespace UE::PixelStreaming2
\ No newline at end of file
diff --git a/Engine/Plugins/Media/PixelStreaming2/Source/PixelStreaming2/Private/VideoProducerViewportBase.cpp b/Engine/Plugins/Media/PixelStreaming2/Source/PixelStreaming2/Private/VideoProducerViewportBase.cpp
new file mode 100644
--- /dev/null
+++ b/Engine/Plugins/Media/PixelStreaming2/Source/PixelStreaming2/Private/VideoProducerViewportBase.cpp
@@ -0,0 +1,118 @@
+// Copyright Epic Games, Inc. All Rights Reserved.
+
+#include "VideoProducerViewportBase.h"
+
+#include "Layout/ArrangedChildren.h"
+#include "Layout/Geometry.h"
+#include "Layout/WidgetPath.h"
+#include "PixelCaptureInputFrameRHI.h"
+
+namespace UE::PixelStreaming2
+{
+	bool FVideoProducerViewportBase::IsFrameAlreadyCopied()
+	{
+#if PLATFORM_MAC
+		return true;
+#else
+		return false;
+#endif
+	}
+	
+	bool FVideoProducerViewportBase::ShouldCaptureViewport()
+	{
+		return false;
+	}
+
+	void FVideoProducerViewportBase::CalculateCaptureRegion(TSharedRef<SViewport> Viewport, TSharedRef<SWindow> Window)
+	{
+		check(IsInGameThread());
+
+		FScopeLock CSLock(&CaptureRectCS);
+	
+		CaptureRect = FIntRect(0, 0, 0, 0);
+	
+		FGeometry InnerWindowGeometry = Window->GetWindowGeometryInWindow();
+	
+		// Find the widget path relative to the window
+		FArrangedChildren JustWindow(EVisibility::Visible);
+		JustWindow.AddWidget(FArrangedWidget(Window, InnerWindowGeometry));
+	
+		FWidgetPath WidgetPath(Window, JustWindow);
+		if (WidgetPath.ExtendPathTo(FWidgetMatcher(Viewport), EVisibility::Visible))
+		{
+			FArrangedWidget ArrangedWidget = WidgetPath.FindArrangedWidget(Viewport).Get(FArrangedWidget::GetNullWidget());
+		
+			FVector2D Position = ArrangedWidget.Geometry.GetAbsolutePosition();
+			FVector2D Size = ArrangedWidget.Geometry.GetAbsoluteSize();
+		
+			CaptureRect = FIntRect(
+				Position.X,
+				Position.Y,
+				Position.X + Size.X,
+				Position.Y + Size.Y);
+		}
+	}
+
+	void FVideoProducerViewportBase::OnBackBufferReadyToPresent(SWindow& Window, const FTextureRHIRef& BackBuffer)
+	{	
+		if(!ShouldCaptureViewport())
+		{
+			return;
+		}
+		
+		if (!TargetWindow.IsValid() || !TargetWindow.HasSameObject(&Window))
+		{
+			return;
+		}
+	
+		FScopeLock CSLock(&CaptureRectCS);
+	
+		if (CaptureRect.IsEmpty())
+		{
+			return;
+		}
+	
+		// Ensure the capture rect in the CopyInfo is in the bound of InputTexture.
+		FIntPoint VirtualPos{
+			FMath::Min(FMath::Max(CaptureRect.Min.X, 0), BackBuffer->GetDesc().Extent.X),
+			FMath::Min(FMath::Max(CaptureRect.Min.Y, 0), BackBuffer->GetDesc().Extent.Y)
+		};
+		FIntPoint VirtualSize{
+			FMath::Min(CaptureRect.Width(), FMath::Max((BackBuffer->GetDesc().Extent.X - VirtualPos.X), 0)),
+			FMath::Min(CaptureRect.Height(), FMath::Max((BackBuffer->GetDesc().Extent.Y - VirtualPos.Y), 0))
+		};
+	
+		if (VirtualSize == FIntPoint::ZeroValue)
+		{
+			return; // discard for copying nothing
+		}
+	
+		TRefCountPtr<IPooledRenderTarget> OutTexture;
+		{
+			FRDGBuilder GraphBuilder(FRHICommandListImmediate::Get());
+			FRDGTextureRef InputTexture = GraphBuilder.RegisterExternalTexture(CreateRenderTarget(BackBuffer, *Window.GetTitle().ToString()));
+			FRDGTextureRef CroppedTexture = GraphBuilder.CreateTexture(FRDGTextureDesc::Create2D(VirtualSize, InputTexture->Desc.Format, FClearValueBinding::None, GetTexCreateFlags()), TEXT("VideoProducerPIEViewport"));
+	
+			FRHICopyTextureInfo CopyInfo;
+			CopyInfo.SourcePosition = FIntVector(VirtualPos.X, VirtualPos.Y, 0);
+			CopyInfo.Size = FIntVector(VirtualSize.X, VirtualSize.Y, 1);
+
+			AddCopyTexturePass(GraphBuilder, InputTexture, CroppedTexture, CopyInfo);
+
+#if PLATFORM_MAC
+			// On Mac specifically, we need to add one more pass to render to a cpu readable texture for AVCodecs
+			FRDGTextureRef StagingTexture = GraphBuilder.CreateTexture(FRDGTextureDesc::Create2D(CroppedTexture->Desc.Extent, EPixelFormat::PF_B8G8R8A8, FClearValueBinding::None, ETextureCreateFlags::CPUReadback), TEXT("VideoProducerPIEViewport MAC Staging"));
+			
+			AddDrawTexturePass(GraphBuilder, GetGlobalShaderMap(GMaxRHIFeatureLevel), CroppedTexture, StagingTexture, FRDGDrawTextureInfo());
+
+			GraphBuilder.QueueTextureExtraction(StagingTexture, &OutTexture);
+#else
+			GraphBuilder.QueueTextureExtraction(CroppedTexture, &OutTexture);
+#endif
+		
+			GraphBuilder.Execute();
+		}
+		
+		PushFrame(FPixelCaptureInputFrameRHI(OutTexture->GetRHI()));
+	}
+} // namespace UE::PixelStreaming2
\ No newline at end of file
diff --git a/Engine/Plugins/Media/PixelStreaming2/Source/PixelStreaming2Editor/Internal/VideoProducerLevelEditor.h b/Engine/Plugins/Media/PixelStreaming2/Source/PixelStreaming2Editor/Internal/VideoProducerLevelEditor.h
new file mode 100644
--- /dev/null
+++ b/Engine/Plugins/Media/PixelStreaming2/Source/PixelStreaming2Editor/Internal/VideoProducerLevelEditor.h
@@ -0,0 +1,35 @@
+// Copyright Epic Games, Inc. All Rights Reserved.
+
+#pragma once
+
+#include "IPixelStreaming2Streamer.h"
+#include "VideoProducerViewportBase.h"
+#include "Widgets/SWindow.h"
+
+class FViewport;
+
+namespace UE::EditorPixelStreaming2
+{
+	using namespace UE::PixelStreaming2;
+
+	/**
+	 * Use this if you want to send the UE primary scene viewport as video input - will only work in editor.
+	 */
+	class PIXELSTREAMING2EDITOR_API FVideoProducerLevelEditor : public FVideoProducerViewportBase
+	{
+	public:
+		static TSharedPtr<FVideoProducerLevelEditor> Create();
+		virtual ~FVideoProducerLevelEditor() = default;
+
+		virtual FString ToString() override;
+
+	protected:
+		virtual bool ShouldCaptureViewport() override;
+
+	private:
+		FVideoProducerLevelEditor() = default;
+
+		void OnPreTick(float DeltaTime);
+	};
+
+} // namespace UE::EditorPixelStreaming2
diff --git a/Engine/Plugins/Media/PixelStreaming2/Source/PixelStreaming2Editor/Internal/VideoProducerViewport.h b/Engine/Plugins/Media/PixelStreaming2/Source/PixelStreaming2Editor/Internal/VideoProducerViewport.h
deleted file mode 100644
diff --git a/Engine/Plugins/Media/PixelStreaming2/Source/PixelStreaming2Editor/Private/PixelStreaming2EditorModule.cpp b/Engine/Plugins/Media/PixelStreaming2/Source/PixelStreaming2Editor/Private/PixelStreaming2EditorModule.cpp
--- a/Engine/Plugins/Media/PixelStreaming2/Source/PixelStreaming2Editor/Private/PixelStreaming2EditorModule.cpp
+++ b/Engine/Plugins/Media/PixelStreaming2/Source/PixelStreaming2Editor/Private/PixelStreaming2EditorModule.cpp
@@ -10,7 +10,7 @@
 #include "Editor/EditorPerformanceSettings.h"
 #include "VideoProducerBackBufferComposited.h"
 #include "VideoProducerPIEViewport.h"
-#include "VideoProducerViewport.h"
+#include "VideoProducerLevelEditor.h"
 #include "PixelStreaming2PluginSettings.h"
 #include "UnrealEngine.h"
 #include "Engine/Engine.h"
@@ -272,7 +272,7 @@ namespace UE::EditorPixelStreaming2
 				InputHandler->SetTargetViewport(SceneViewport->GetViewportWidget());
 				InputHandler->SetTargetWindow(SceneViewport->FindWindow());
 				InputHandler->SetInputType(EPixelStreaming2InputType::RouteToWindow);
-				EditorStreamer->SetVideoProducer(FVideoProducerViewport::Create(EditorStreamer));
+				EditorStreamer->SetVideoProducer(FVideoProducerLevelEditor::Create());
 			}
 			break;
 			case EPixelStreaming2EditorStreamTypes::Editor:
diff --git a/Engine/Plugins/Media/PixelStreaming2/Source/PixelStreaming2Editor/Private/VideoProducerBackBufferComposited.cpp b/Engine/Plugins/Media/PixelStreaming2/Source/PixelStreaming2Editor/Private/VideoProducerBackBufferComposited.cpp
--- a/Engine/Plugins/Media/PixelStreaming2/Source/PixelStreaming2Editor/Private/VideoProducerBackBufferComposited.cpp
+++ b/Engine/Plugins/Media/PixelStreaming2/Source/PixelStreaming2Editor/Private/VideoProducerBackBufferComposited.cpp
@@ -15,21 +15,6 @@ namespace UE::EditorPixelStreaming2
 	DECLARE_LOG_CATEGORY_EXTERN(LogPixelStreaming2BackBufferComposited, Log, VeryVerbose);
 	DEFINE_LOG_CATEGORY(LogPixelStreaming2BackBufferComposited);
 
-	ETextureCreateFlags GetTexCreateFlags()
-	{
-		ETextureCreateFlags TexCreateFlags = ETextureCreateFlags::RenderTargetable;
-		if (RHIGetInterfaceType() == ERHIInterfaceType::Vulkan)
-		{
-			TexCreateFlags |= ETextureCreateFlags::External;
-		}
-		else
-		{
-			TexCreateFlags |= ETextureCreateFlags::Shared;
-		}
-		
-		return TexCreateFlags;
-	}
-
 	TSharedPtr<FVideoProducerBackBufferComposited> FVideoProducerBackBufferComposited::Create()
 	{
 		TSharedPtr<FVideoProducerBackBufferComposited> NewInput = TSharedPtr<FVideoProducerBackBufferComposited>(new FVideoProducerBackBufferComposited());
diff --git a/Engine/Plugins/Media/PixelStreaming2/Source/PixelStreaming2Editor/Private/VideoProducerLevelEditor.cpp b/Engine/Plugins/Media/PixelStreaming2/Source/PixelStreaming2Editor/Private/VideoProducerLevelEditor.cpp
new file mode 100644
--- /dev/null
+++ b/Engine/Plugins/Media/PixelStreaming2/Source/PixelStreaming2Editor/Private/VideoProducerLevelEditor.cpp
@@ -0,0 +1,75 @@
+// Copyright Epic Games, Inc. All Rights Reserved.
+
+#include "VideoProducerLevelEditor.h"
+
+#include "Engine/GameViewportClient.h"
+#include "IPixelStreaming2InputHandler.h"
+#include "LevelEditor.h"
+#include "LevelEditorViewport.h"
+#include "Modules/ModuleManager.h"
+#include "PixelCaptureInputFrameRHI.h"
+#include "SLevelViewport.h"
+#include "Framework/Application/SlateApplication.h"
+#include "Rendering/SlateRenderer.h"
+#include "UtilsAsync.h"
+
+namespace UE::EditorPixelStreaming2
+{
+	TSharedPtr<FVideoProducerLevelEditor> FVideoProducerLevelEditor::Create()
+	{
+		TSharedPtr<FVideoProducerLevelEditor> NewInput = TSharedPtr<FVideoProducerLevelEditor>(new FVideoProducerLevelEditor());
+		TWeakPtr<FVideoProducerLevelEditor> WeakInput = NewInput;
+
+		UE::PixelStreaming2::DoOnGameThread([WeakInput]() {
+			if (TSharedPtr<FVideoProducerLevelEditor> Input = WeakInput.Pin())
+			{
+				FSlateRenderer* Renderer = FSlateApplication::Get().GetRenderer();
+				if(!Renderer)
+				{
+					return;
+				}
+ 
+				Renderer->OnBackBufferReadyToPresent().AddSP(Input.ToSharedRef(), &FVideoProducerLevelEditor::OnBackBufferReadyToPresent);
+				FSlateApplication::Get().OnPreTick().AddSP(Input.ToSharedRef(), &FVideoProducerLevelEditor::OnPreTick);
+			}
+		});
+
+		return NewInput;
+	}
+
+	FString FVideoProducerLevelEditor::ToString()
+	{
+		return TEXT("the Target Viewport");
+	}
+
+	bool FVideoProducerLevelEditor::ShouldCaptureViewport()
+	{
+		return true;
+	}
+
+	void FVideoProducerLevelEditor::OnPreTick(float DeltaTime)
+	{
+		if (!ShouldCaptureViewport())
+		{
+			return;
+		}
+
+		FLevelEditorModule&		   LevelEditorModule = FModuleManager::GetModuleChecked<FLevelEditorModule>("LevelEditor");
+		TSharedPtr<SLevelViewport> ActiveLevelViewport = LevelEditorModule.GetFirstActiveLevelViewport();
+		if (!ActiveLevelViewport.IsValid())
+		{
+			return;
+		}
+
+		TargetViewport = ActiveLevelViewport->GetViewportWidget();
+		if (TargetViewport.IsValid())
+		{
+			TargetWindow = FSlateApplication::Get().FindWidgetWindow(TargetViewport.Pin().ToSharedRef());
+			if (TargetWindow.IsValid())
+			{
+				CalculateCaptureRegion(TargetViewport.Pin().ToSharedRef(), TargetWindow.Pin().ToSharedRef());
+			}
+		}
+	}
+
+} // namespace UE::EditorPixelStreaming2
\ No newline at end of file
diff --git a/Engine/Plugins/Media/PixelStreaming2/Source/PixelStreaming2Editor/Private/VideoProducerViewport.cpp b/Engine/Plugins/Media/PixelStreaming2/Source/PixelStreaming2Editor/Private/VideoProducerViewport.cpp
deleted file mode 100644
