[PS2] Fix: missing include

diff --git a/Engine/Plugins/Media/PixelStreaming2/Source/PixelStreaming2/Private/VideoProducerViewportBase.cpp b/Engine/Plugins/Media/PixelStreaming2/Source/PixelStreaming2/Private/VideoProducerViewportBase.cpp
--- a/Engine/Plugins/Media/PixelStreaming2/Source/PixelStreaming2/Private/VideoProducerViewportBase.cpp
+++ b/Engine/Plugins/Media/PixelStreaming2/Source/PixelStreaming2/Private/VideoProducerViewportBase.cpp
@@ -2,6 +2,7 @@
 
 #include "VideoProducerViewportBase.h"
 
+#include "GlobalShader.h"
 #include "Layout/ArrangedChildren.h"
 #include "Layout/Geometry.h"
 #include "Layout/WidgetPath.h"
