[PS1, PS2] Fix a crash in editor when setting the streamID on the command line

diff --git a/Engine/Plugins/Media/PixelStreaming/Source/PixelStreamingEditor/Private/PixelStreamingEditorModule.cpp b/Engine/Plugins/Media/PixelStreaming/Source/PixelStreamingEditor/Private/PixelStreamingEditorModule.cpp
--- a/Engine/Plugins/Media/PixelStreaming/Source/PixelStreamingEditor/Private/PixelStreamingEditorModule.cpp
+++ b/Engine/Plugins/Media/PixelStreaming/Source/PixelStreamingEditor/Private/PixelStreamingEditorModule.cpp
@@ -77,7 +77,6 @@ void FPixelStreamingEditorModule::ShutdownModule()
 
 void FPixelStreamingEditorModule::InitEditorStreaming(IPixelStreamingModule& Module)
 {
-	FString EditorStreamerID;
 	if (!FParse::Value(FCommandLine::Get(), TEXT("PixelStreamingID="), EditorStreamerID))
 	{
 		EditorStreamerID = "Editor";
diff --git a/Engine/Plugins/Media/PixelStreaming/Source/PixelStreamingEditor/Private/PixelStreamingEditorModule.h b/Engine/Plugins/Media/PixelStreaming/Source/PixelStreamingEditor/Private/PixelStreamingEditorModule.h
--- a/Engine/Plugins/Media/PixelStreaming/Source/PixelStreamingEditor/Private/PixelStreamingEditorModule.h
+++ b/Engine/Plugins/Media/PixelStreaming/Source/PixelStreamingEditor/Private/PixelStreamingEditorModule.h
@@ -38,6 +38,9 @@ public:
 
 	virtual bool UseExternalSignallingServer() override;
 	virtual void UseExternalSignallingServer(bool bInUseExternalSignallingServer) override;
+	/** End IModuleInterface implementation */
+
+	virtual FString GetEditorStreamerID() const { return EditorStreamerID; }
 
 private:
 	void InitEditorStreaming(IPixelStreamingModule& Module);
@@ -54,6 +57,8 @@ private:
 	TSharedPtr<FMonitoredProcess> DownloadProcess;
 	// The signalling server host: eg ws://127.0.0.1
 	FString SignallingDomain;
+	// The ID for the Editor stream, eg "Editor"
+	FString EditorStreamerID;
 	// The port the streamer will connect to. eg 8888
 	int32 StreamerPort;
 	// The port the streams can be viewed at on the browser. eg 80 or 8080
diff --git a/Engine/Plugins/Media/PixelStreaming/Source/PixelStreamingEditor/Private/PixelStreamingToolbar.cpp b/Engine/Plugins/Media/PixelStreaming/Source/PixelStreamingEditor/Private/PixelStreamingToolbar.cpp
--- a/Engine/Plugins/Media/PixelStreaming/Source/PixelStreamingEditor/Private/PixelStreamingToolbar.cpp
+++ b/Engine/Plugins/Media/PixelStreaming/Source/PixelStreamingEditor/Private/PixelStreamingToolbar.cpp
@@ -76,7 +76,8 @@ namespace UE::EditorPixelStreaming
 				IPixelStreamingEditorModule::Get().StartStreaming(EStreamTypes::LevelEditorViewport);
 			}),
 			FCanExecuteAction::CreateLambda([] {
-				if (TSharedPtr<IPixelStreamingStreamer> Streamer = IPixelStreamingModule::Get().FindStreamer("Editor"))
+				FPixelStreamingEditorModule& EditorModule = static_cast<FPixelStreamingEditorModule&>(IPixelStreamingEditorModule::Get());
+				if (TSharedPtr<IPixelStreamingStreamer> Streamer = IPixelStreamingModule::Get().FindStreamer(EditorModule.GetEditorStreamerID()))
 				{
 					return !Streamer->IsStreaming();
 				}
@@ -89,7 +90,8 @@ namespace UE::EditorPixelStreaming
 				IPixelStreamingEditorModule::Get().StartStreaming(EStreamTypes::Editor);
 			}),
 			FCanExecuteAction::CreateLambda([] {
-				if (TSharedPtr<IPixelStreamingStreamer> Streamer = IPixelStreamingModule::Get().FindStreamer("Editor"))
+				FPixelStreamingEditorModule& EditorModule = static_cast<FPixelStreamingEditorModule&>(IPixelStreamingEditorModule::Get());
+				if (TSharedPtr<IPixelStreamingStreamer> Streamer = IPixelStreamingModule::Get().FindStreamer(EditorModule.GetEditorStreamerID()))
 				{
 					return !Streamer->IsStreaming();
 				}
@@ -339,7 +341,8 @@ namespace UE::EditorPixelStreaming
 					  .AutoWidth()
 						  [SNew(SEditableTextBox)
 								  .Text_Lambda([]() {
-									  TSharedPtr<IPixelStreamingStreamer> Streamer = IPixelStreamingModule::Get().FindStreamer("Editor");
+									  FPixelStreamingEditorModule& EditorModule = static_cast<FPixelStreamingEditorModule&>(IPixelStreamingEditorModule::Get());
+									  TSharedPtr<IPixelStreamingStreamer> Streamer = IPixelStreamingModule::Get().FindStreamer(EditorModule.GetEditorStreamerID());
 									  return FText::FromString(Streamer->GetSignallingServerURL());
 								  })
 								  .OnTextChanged_Lambda([](const FText& InText) {
@@ -433,7 +436,8 @@ namespace UE::EditorPixelStreaming
 				}
 				else
 				{
-					if (Streamer->GetId() == "Editor")
+					FPixelStreamingEditorModule& EditorModule = static_cast<FPixelStreamingEditorModule&>(IPixelStreamingEditorModule::Get());
+					if (Streamer->GetId() == EditorModule.GetEditorStreamerID())
 					{
 						MenuBuilder.AddMenuEntry(FPixelStreamingCommands::Get().StreamLevelEditor);
 						MenuBuilder.AddMenuEntry(FPixelStreamingCommands::Get().StreamEditor);
diff --git a/Engine/Plugins/Media/PixelStreaming2/Source/PixelStreaming2Editor/Private/PixelStreaming2EditorModule.cpp b/Engine/Plugins/Media/PixelStreaming2/Source/PixelStreaming2Editor/Private/PixelStreaming2EditorModule.cpp
--- a/Engine/Plugins/Media/PixelStreaming2/Source/PixelStreaming2Editor/Private/PixelStreaming2EditorModule.cpp
+++ b/Engine/Plugins/Media/PixelStreaming2/Source/PixelStreaming2Editor/Private/PixelStreaming2EditorModule.cpp
@@ -148,7 +148,6 @@ namespace UE::EditorPixelStreaming2
 
 	void FPixelStreaming2EditorModule::InitEditorStreaming(IPixelStreaming2RTCModule& Module)
 	{
-		FString EditorStreamerID;
 		if (!FParse::Value(FCommandLine::Get(), TEXT("PixelStreamingID="), EditorStreamerID))
 		{
 			EditorStreamerID = "Editor";
diff --git a/Engine/Plugins/Media/PixelStreaming2/Source/PixelStreaming2Editor/Private/PixelStreaming2EditorModule.h b/Engine/Plugins/Media/PixelStreaming2/Source/PixelStreaming2Editor/Private/PixelStreaming2EditorModule.h
--- a/Engine/Plugins/Media/PixelStreaming2/Source/PixelStreaming2Editor/Private/PixelStreaming2EditorModule.h
+++ b/Engine/Plugins/Media/PixelStreaming2/Source/PixelStreaming2Editor/Private/PixelStreaming2EditorModule.h
@@ -37,6 +37,9 @@ namespace UE::EditorPixelStreaming2
 		virtual FString GetSSLCertificatePath() override { return SSLCertificatePath; }
 		virtual void	SetSSLPrivateKeyPath(const FString& Path) override;
 		virtual FString GetSSLPrivateKeyPath() override { return SSLPrivateKeyPath; }
+		/** End IModuleInterface implementation */
+
+		virtual FString GetEditorStreamerID() const { return EditorStreamerID; }
 
 	private:
 		void InitEditorStreaming(IPixelStreaming2RTCModule& Module);
@@ -72,6 +75,8 @@ namespace UE::EditorPixelStreaming2
 		//
 		FString SSLPrivateKeyPath;
 
+		FString EditorStreamerID;
+
 		// The streamer used by the PixelStreaming2Editor module
 		TSharedPtr<IPixelStreaming2Streamer> EditorStreamer;
 
diff --git a/Engine/Plugins/Media/PixelStreaming2/Source/PixelStreaming2Editor/Private/PixelStreaming2Toolbar.cpp b/Engine/Plugins/Media/PixelStreaming2/Source/PixelStreaming2Editor/Private/PixelStreaming2Toolbar.cpp
--- a/Engine/Plugins/Media/PixelStreaming2/Source/PixelStreaming2Editor/Private/PixelStreaming2Toolbar.cpp
+++ b/Engine/Plugins/Media/PixelStreaming2/Source/PixelStreaming2Editor/Private/PixelStreaming2Toolbar.cpp
@@ -164,7 +164,8 @@ namespace UE::EditorPixelStreaming2
 				IPixelStreaming2EditorModule::Get().StartStreaming(EPixelStreaming2EditorStreamTypes::LevelEditorViewport);
 			}),
 			FCanExecuteAction::CreateLambda([] {
-				if (TSharedPtr<IPixelStreaming2Streamer> Streamer = IPixelStreaming2Module::Get().FindStreamer("Editor"))
+				FPixelStreaming2EditorModule& EditorModule = static_cast<FPixelStreaming2EditorModule&>(IPixelStreaming2EditorModule::Get());
+				if (TSharedPtr<IPixelStreaming2Streamer> Streamer = IPixelStreaming2Module::Get().FindStreamer(EditorModule.GetEditorStreamerID()))
 				{
 					return !Streamer->IsStreaming();
 				}
@@ -177,7 +178,8 @@ namespace UE::EditorPixelStreaming2
 				IPixelStreaming2EditorModule::Get().StartStreaming(EPixelStreaming2EditorStreamTypes::Editor);
 			}),
 			FCanExecuteAction::CreateLambda([] {
-				if (TSharedPtr<IPixelStreaming2Streamer> Streamer = IPixelStreaming2Module::Get().FindStreamer("Editor"))
+				FPixelStreaming2EditorModule& EditorModule = static_cast<FPixelStreaming2EditorModule&>(IPixelStreaming2EditorModule::Get());
+				if (TSharedPtr<IPixelStreaming2Streamer> Streamer = IPixelStreaming2Module::Get().FindStreamer(EditorModule.GetEditorStreamerID()))
 				{
 					return !Streamer->IsStreaming();
 				}
@@ -451,8 +453,9 @@ namespace UE::EditorPixelStreaming2
 					  .AutoWidth()
 						  [SNew(SEditableTextBox)
 								  .Text_Lambda([]() {
-									  TSharedPtr<IPixelStreaming2Streamer> Streamer = IPixelStreaming2Module::Get().FindStreamer("Editor");
-									  return FText::FromString(Streamer->GetConnectionURL());
+									   FPixelStreaming2EditorModule& EditorModule = static_cast<FPixelStreaming2EditorModule&>(IPixelStreaming2EditorModule::Get());
+									   TSharedPtr<IPixelStreaming2Streamer> Streamer = IPixelStreaming2Module::Get().FindStreamer(EditorModule.GetEditorStreamerID());
+									   return FText::FromString(Streamer->GetConnectionURL());
 								  })
 								  .OnTextChanged_Lambda([](const FText& InText) {
 									  IPixelStreaming2Module::Get().ForEachStreamer([InText](TSharedPtr<IPixelStreaming2Streamer> Streamer) {
@@ -545,7 +548,8 @@ namespace UE::EditorPixelStreaming2
 				}
 				else
 				{
-					if (Streamer->GetId() == "Editor")
+					FPixelStreaming2EditorModule& EditorModule = static_cast<FPixelStreaming2EditorModule&>(IPixelStreaming2EditorModule::Get());
+					if (Streamer->GetId() == EditorModule.GetEditorStreamerID())
 					{
 						MenuBuilder.AddMenuEntry(FPixelStreaming2Commands::Get().StreamLevelEditor);
 						MenuBuilder.AddMenuEntry(FPixelStreaming2Commands::Get().StreamEditor);
