[PS2] Fix: PIE streaming not working as it's unable to get a viewport

#FailedValidations Preflight

diff --git a/Engine/Plugins/Media/PixelStreaming2/Source/PixelStreaming2Editor/Private/PixelStreaming2EditorModule.cpp b/Engine/Plugins/Media/PixelStreaming2/Source/PixelStreaming2Editor/Private/PixelStreaming2EditorModule.cpp
--- a/Engine/Plugins/Media/PixelStreaming2/Source/PixelStreaming2Editor/Private/PixelStreaming2EditorModule.cpp
+++ b/Engine/Plugins/Media/PixelStreaming2/Source/PixelStreaming2Editor/Private/PixelStreaming2EditorModule.cpp
@@ -563,21 +563,27 @@ namespace UE::EditorPixelStreaming2
 			return;
 		}
 
-		FWorldContext* PIEWorldContext = GEditor->GetPIEWorldContext();
-		if (!PIEWorldContext)
+		UGameViewportClient* Viewport = nullptr;
+		// Iterate through world contexts to get the first PIE viewport that is for a client or a standalone game
+		for (const FWorldContext& Context : GEditor->GetWorldContexts())
 		{
-			UE_LOGFMT(LogPixelStreaming2Editor, Warning, "Failed to get editor PIE world context. PIE streamer will not stream!");
-			return;
+			if (Context.WorldType == EWorldType::PIE && Context.World())
+			{
+				ENetMode NetMode = Context.World()->GetNetMode();
+        		if (NetMode == NM_Client || NetMode == NM_Standalone)
+        		{
+            		Viewport = Context.GameViewport;
+            		break;
+        		}
+			}
 		}
 
-		UWorld* World = PIEWorldContext->World();
-		if (!World)
+		if (!Viewport)
 		{
-			UE_LOGFMT(LogPixelStreaming2Editor, Warning, "Failed to get editor PIE world. PIE streamer will not stream!");
+			UE_LOGFMT(LogPixelStreaming2Editor, Warning, "Failed to get PIE viewport. PIE streamer will not stream!");
 			return;
 		}
 
-		UGameViewportClient* Viewport = World->GetGameViewport();
 		InputHandler->SetTargetViewport(Viewport->GetGameViewportWidget());
 		InputHandler->SetTargetWindow(Viewport->GetWindow());
 		InputHandler->SetInputType(EPixelStreaming2InputType::RouteToWindow);
