The source files included were modified by the UnrealCodeFixup tool so that they can pass the -mergemodules compilation.

diff --git a/Engine/Plugins/Media/PixelStreaming2/Source/PixelStreaming2/Internal/Blueprints/PixelStreaming2VideoProducer.h b/Engine/Plugins/Media/PixelStreaming2/Source/PixelStreaming2/Internal/Blueprints/PixelStreaming2VideoProducer.h
--- a/Engine/Plugins/Media/PixelStreaming2/Source/PixelStreaming2/Internal/Blueprints/PixelStreaming2VideoProducer.h
+++ b/Engine/Plugins/Media/PixelStreaming2/Source/PixelStreaming2/Internal/Blueprints/PixelStreaming2VideoProducer.h
@@ -7,11 +7,13 @@
 
 #include "PixelStreaming2VideoProducer.generated.h"
 
+#define UE_API PIXELSTREAMING2_API
+
 /**
  * The base video producer blueprint class
  */
-UCLASS(NotBlueprintType, NotBlueprintable, Category = "PixelStreaming2", META = (DisplayName = "Video Producer"))
-class PIXELSTREAMING2_API UPixelStreaming2VideoProducerBase : public UObject
+UCLASS(MinimalAPI, NotBlueprintType, NotBlueprintable, Category = "PixelStreaming2", META = (DisplayName = "Video Producer"))
+class UPixelStreaming2VideoProducerBase : public UObject
 {
 	GENERATED_UCLASS_BODY()
 
@@ -25,32 +27,32 @@ protected:
 /**
  * A video producer that streams the Unreal Engine's back buffer
  */
-UCLASS(NotBlueprintType, NotBlueprintable, Category = "PixelStreaming2", META = (DisplayName = "Back Buffer Video Producer"))
-class PIXELSTREAMING2_API UPixelStreaming2VideoProducerBackBuffer : public UPixelStreaming2VideoProducerBase
+UCLASS(MinimalAPI, NotBlueprintType, NotBlueprintable, Category = "PixelStreaming2", META = (DisplayName = "Back Buffer Video Producer"))
+class UPixelStreaming2VideoProducerBackBuffer : public UPixelStreaming2VideoProducerBase
 {
 	GENERATED_UCLASS_BODY()
 
 public:
-	virtual TSharedPtr<IPixelStreaming2VideoProducer> GetVideoProducer() override;
+	UE_API virtual TSharedPtr<IPixelStreaming2VideoProducer> GetVideoProducer() override;
 };
 
 /**
  * A video producer that streams the Unreal Engine's back buffer but uses the MediaIO Framework to capture the frame rather than Pixel Capture
  */
-UCLASS(NotBlueprintType, NotBlueprintable, Category = "PixelStreaming2", META = (DisplayName = "Media Capture Video Producer"))
-class PIXELSTREAMING2_API UPixelStreaming2VideoProducerMediaCapture : public UPixelStreaming2VideoProducerBase
+UCLASS(MinimalAPI, NotBlueprintType, NotBlueprintable, Category = "PixelStreaming2", META = (DisplayName = "Media Capture Video Producer"))
+class UPixelStreaming2VideoProducerMediaCapture : public UPixelStreaming2VideoProducerBase
 {
 	GENERATED_UCLASS_BODY()
 
 public:
-	virtual TSharedPtr<IPixelStreaming2VideoProducer> GetVideoProducer() override;
+	UE_API virtual TSharedPtr<IPixelStreaming2VideoProducer> GetVideoProducer() override;
 };
 
 /**
  * A video producer that streams the specified render target
  */
-UCLASS(NotBlueprintType, NotBlueprintable, Category = "PixelStreaming2", META = (DisplayName = "Render Target Video Producer"))
-class PIXELSTREAMING2_API UPixelStreaming2VideoProducerRenderTarget : public UPixelStreaming2VideoProducerBase
+UCLASS(MinimalAPI, NotBlueprintType, NotBlueprintable, Category = "PixelStreaming2", META = (DisplayName = "Render Target Video Producer"))
+class UPixelStreaming2VideoProducerRenderTarget : public UPixelStreaming2VideoProducerBase
 {
 	GENERATED_UCLASS_BODY()
 
@@ -58,5 +60,7 @@ public:
 	UPROPERTY(BlueprintReadWrite, EditAnywhere, Category = Video, AssetRegistrySearchable)
 	TObjectPtr<UTextureRenderTarget2D> Target;
 
-	virtual TSharedPtr<IPixelStreaming2VideoProducer> GetVideoProducer() override;
-};
\ No newline at end of file
+	UE_API virtual TSharedPtr<IPixelStreaming2VideoProducer> GetVideoProducer() override;
+};
+
+#undef UE_API
diff --git a/Engine/Plugins/Media/PixelStreaming2/Source/PixelStreaming2/Internal/PixelStreaming2MediaIOCapture.h b/Engine/Plugins/Media/PixelStreaming2/Source/PixelStreaming2/Internal/PixelStreaming2MediaIOCapture.h
--- a/Engine/Plugins/Media/PixelStreaming2/Source/PixelStreaming2/Internal/PixelStreaming2MediaIOCapture.h
+++ b/Engine/Plugins/Media/PixelStreaming2/Source/PixelStreaming2/Internal/PixelStreaming2MediaIOCapture.h
@@ -9,8 +9,10 @@
 
 #include "PixelStreaming2MediaIOCapture.generated.h"
 
-UCLASS(BlueprintType)
-class PIXELSTREAMING2_API UPixelStreaming2MediaIOCapture : public UMediaCapture
+#define UE_API PIXELSTREAMING2_API
+
+UCLASS(MinimalAPI, BlueprintType)
+class UPixelStreaming2MediaIOCapture : public UMediaCapture
 {
 	GENERATED_BODY()
 
@@ -19,7 +21,7 @@ public:
 	/**
 	 * GPU copy methods
 	 */
-	virtual void OnRHIResourceCaptured_RenderingThread(
+	UE_API virtual void OnRHIResourceCaptured_RenderingThread(
 		FRHICommandListImmediate&							   RHICmdList,
 		const FCaptureBaseData&								   InBaseData,
 		TSharedPtr<FMediaCaptureUserData, ESPMode::ThreadSafe> InUserData,
@@ -28,7 +30,7 @@ public:
 	/**
 	 * Custom conversion operation for Mac
 	 */
-	virtual void OnCustomCapture_RenderingThread(
+	UE_API virtual void OnCustomCapture_RenderingThread(
 		FRDGBuilder&										   GraphBuilder,
 		const FCaptureBaseData&								   InBaseData,
 		TSharedPtr<FMediaCaptureUserData, ESPMode::ThreadSafe> InUserData,
@@ -38,18 +40,18 @@ public:
 		FVector2D											   CropU,
 		FVector2D											   CropV) override;
 
-	virtual bool InitializeCapture() override;
-	virtual bool PostInitializeCaptureViewport(TSharedPtr<FSceneViewport>& InSceneViewport) override;
+	UE_API virtual bool InitializeCapture() override;
+	UE_API virtual bool PostInitializeCaptureViewport(TSharedPtr<FSceneViewport>& InSceneViewport) override;
 	virtual bool ShouldCaptureRHIResource() const { return bDoGPUCopy; }
-	virtual void StopCaptureImpl(bool bAllowPendingFrameToBeProcess) override;
+	UE_API virtual void StopCaptureImpl(bool bAllowPendingFrameToBeProcess) override;
 	// We override the texture flags because on Mac we want the texture to have the CPU_Readback flag
-	virtual ETextureCreateFlags GetOutputTextureFlags() const override;
-	virtual void WaitForGPU(FRHITexture* InRHITexture) override;
-	virtual TSharedPtr<FMediaCaptureUserData, ESPMode::ThreadSafe> GetCaptureFrameUserData_GameThread() override;
+	UE_API virtual ETextureCreateFlags GetOutputTextureFlags() const override;
+	UE_API virtual void WaitForGPU(FRHITexture* InRHITexture) override;
+	UE_API virtual TSharedPtr<FMediaCaptureUserData, ESPMode::ThreadSafe> GetCaptureFrameUserData_GameThread() override;
 	//~ End UMediaCapture interface
 
 	TSharedPtr<FSceneViewport>				GetViewport() const { return SceneViewport.Pin(); }
-	virtual void							ViewportResized(FViewport* Viewport, uint32 ResizeCode);
+	UE_API virtual void							ViewportResized(FViewport* Viewport, uint32 ResizeCode);
 	bool									WasViewportResized() const { return bViewportResized; }
 	void									SetVideoProducer(TSharedPtr<IPixelStreaming2VideoProducer> InVideoProducer) { VideoProducer = InVideoProducer; }
 	TWeakPtr<IPixelStreaming2VideoProducer> GetVideoProducer() { return VideoProducer; }
@@ -58,8 +60,8 @@ public:
 	FOnCaptureViewportInitialized OnCaptureViewportInitialized;
 
 private:
-	void HandleCapturedFrame(FTextureRHIRef InTexture, TSharedPtr<FMediaCaptureUserData, ESPMode::ThreadSafe> InUserData);
-	void UpdateCaptureResolution(int32 Width, int32 Height);
+	UE_API void HandleCapturedFrame(FTextureRHIRef InTexture, TSharedPtr<FMediaCaptureUserData, ESPMode::ThreadSafe> InUserData);
+	UE_API void UpdateCaptureResolution(int32 Width, int32 Height);
 
 private:
 	TWeakPtr<FSceneViewport>				SceneViewport;
@@ -75,4 +77,6 @@ private:
 	bool bDoGPUCopy = true;
 
 	TDeque<FGPUFenceRHIRef> Fences;
-};
\ No newline at end of file
+};
+
+#undef UE_API
diff --git a/Engine/Plugins/Media/PixelStreaming2/Source/PixelStreaming2/Internal/SoundGenerator.h b/Engine/Plugins/Media/PixelStreaming2/Source/PixelStreaming2/Internal/SoundGenerator.h
--- a/Engine/Plugins/Media/PixelStreaming2/Source/PixelStreaming2/Internal/SoundGenerator.h
+++ b/Engine/Plugins/Media/PixelStreaming2/Source/PixelStreaming2/Internal/SoundGenerator.h
@@ -6,22 +6,24 @@
 
 #include <atomic>
 
+#define UE_API PIXELSTREAMING2_API
+
 namespace UE::PixelStreaming2
 {
 	/**
 	 * An `ISoundGenerator` implementation to pump some audio from EpicRtc into this synth component
 	 */
-	class PIXELSTREAMING2_API FSoundGenerator : public ::ISoundGenerator
+	class FSoundGenerator : public ::ISoundGenerator
 	{
 	public:
-		FSoundGenerator();
+		UE_API FSoundGenerator();
 		virtual ~FSoundGenerator() = default;
 
 		// Called when a new buffer is required.
-		virtual int32 OnGenerateAudio(float* OutAudio, int32 NumSamples) override;
+		UE_API virtual int32 OnGenerateAudio(float* OutAudio, int32 NumSamples) override;
 
 		// Returns the number of samples to render per callback
-		virtual int32 GetDesiredNumSamplesToRenderPerCallback() const;
+		UE_API virtual int32 GetDesiredNumSamplesToRenderPerCallback() const;
 
 		// Optional. Called on audio generator thread right when the generator begins generating.
 		virtual void OnBeginGenerate() { bGeneratingAudio = true; };
@@ -32,12 +34,12 @@ namespace UE::PixelStreaming2
 		// Optional. Can be overridden to end the sound when generating is finished.
 		virtual bool IsFinished() const { return false; };
 
-		void AddAudio(const int16_t* AudioData, int InSampleRate, size_t NChannels, size_t NFrames);
+		UE_API void AddAudio(const int16_t* AudioData, int InSampleRate, size_t NChannels, size_t NFrames);
 
 		int32 GetSampleRate() { return Params.SampleRate; }
 		int32 GetNumChannels() { return Params.NumChannels; }
-		void  EmptyBuffers();
-		void  SetParameters(const FSoundGeneratorInitParams& InitParams);
+		UE_API void  EmptyBuffers();
+		UE_API void  SetParameters(const FSoundGeneratorInitParams& InitParams);
 
 	private:
 		FSoundGeneratorInitParams Params;
@@ -50,3 +52,5 @@ namespace UE::PixelStreaming2
 	};
 
 } // namespace UE::PixelStreaming2
+
+#undef UE_API
diff --git a/Engine/Plugins/Media/PixelStreaming2/Source/PixelStreaming2/Internal/StreamerReconnectTimer.h b/Engine/Plugins/Media/PixelStreaming2/Source/PixelStreaming2/Internal/StreamerReconnectTimer.h
--- a/Engine/Plugins/Media/PixelStreaming2/Source/PixelStreaming2/Internal/StreamerReconnectTimer.h
+++ b/Engine/Plugins/Media/PixelStreaming2/Source/PixelStreaming2/Internal/StreamerReconnectTimer.h
@@ -6,6 +6,8 @@
 #include "Tickable.h"
 #include "IPixelStreaming2Streamer.h"
 
+#define UE_API PIXELSTREAMING2_API
+
 namespace UE::PixelStreaming2
 {
 	/**
@@ -14,17 +16,17 @@ namespace UE::PixelStreaming2
 	 * Note: this timer runs on the game thread and its reconnect interval is controlled by the CVar
 	 * UPixelStreaming2PluginSettings::CVarSignalingReconnectInterval
 	 */
-	class PIXELSTREAMING2_API FStreamerReconnectTimer : public FTickableGameObject
+	class FStreamerReconnectTimer : public FTickableGameObject
 	{
 	public:
 		FStreamerReconnectTimer() = default;
 		virtual ~FStreamerReconnectTimer() = default;
-		void Start(TWeakPtr<IPixelStreaming2Streamer> InWeakStreamer);
-		void Stop();
-		void Reset();
+		UE_API void Start(TWeakPtr<IPixelStreaming2Streamer> InWeakStreamer);
+		UE_API void Stop();
+		UE_API void Reset();
 
 		// Begin FTickableGameObject
-		virtual void		Tick(float DeltaTime) override;
+		UE_API virtual void		Tick(float DeltaTime) override;
 		virtual bool		IsTickableInEditor() const override { return true; }
 		virtual bool		IsTickableWhenPaused() const override { return true; }
 		FORCEINLINE TStatId GetStatId() const { RETURN_QUICK_DECLARE_CYCLE_STAT(PixelStreaming2StreamerReconnectTimer, STATGROUP_Tickables); }
@@ -40,4 +42,6 @@ namespace UE::PixelStreaming2
 		uint64							   NumReconnectAttempts = 0;
 		uint64							   LastReconnectCycles = 0;
 	};
-} // namespace UE::PixelStreaming2
\ No newline at end of file
+} // namespace UE::PixelStreaming2
+
+#undef UE_API
diff --git a/Engine/Plugins/Media/PixelStreaming2/Source/PixelStreaming2/Internal/TickableTask.h b/Engine/Plugins/Media/PixelStreaming2/Source/PixelStreaming2/Internal/TickableTask.h
--- a/Engine/Plugins/Media/PixelStreaming2/Source/PixelStreaming2/Internal/TickableTask.h
+++ b/Engine/Plugins/Media/PixelStreaming2/Source/PixelStreaming2/Internal/TickableTask.h
@@ -5,6 +5,8 @@
 #include "Templates/UniquePtr.h"
 #include "Templates/SharedPointer.h"
 
+#define UE_API PIXELSTREAMING2_API
+
 namespace UE::PixelStreaming2
 {
 	template <typename T>
@@ -14,7 +16,7 @@ namespace UE::PixelStreaming2
 	 * Base class for a tickable task. Inheriting from this class ensures that your task can be ticked by the PixelStreaming thread.
 	 * Note it is important to destroy the task when it is no longer in use to save the tick thread.
 	 */
-	class PIXELSTREAMING2_API FPixelStreamingTickableTask : public TSharedFromThis<FPixelStreamingTickableTask>
+	class FPixelStreamingTickableTask : public TSharedFromThis<FPixelStreamingTickableTask>
 	{
 	public:
 		/**
@@ -36,8 +38,8 @@ namespace UE::PixelStreaming2
 		FPixelStreamingTickableTask() = default;
 
 	private:
-		void Register();
-		void Unregister();
+		UE_API void Register();
+		UE_API void Unregister();
 
 		// Allow FPixelStreamingTaskDeleter to access the Unregiser method
 		friend FPixelStreamingTaskDeleter;
@@ -47,7 +49,7 @@ namespace UE::PixelStreaming2
 	 * TUniquePtr custom deleter that handle automatic unregistering of the task.
 	 * This ensure that the task won't try to be ticked mid-deletion
 	 */
-	struct PIXELSTREAMING2_API FPixelStreamingTaskDeleter
+	struct FPixelStreamingTaskDeleter
 	{
 		template <typename T>
 		void operator()(T* Ptr) const
@@ -69,4 +71,6 @@ namespace UE::PixelStreaming2
 		return Task;
 	}
 
-} // namespace UE::PixelStreaming2
\ No newline at end of file
+} // namespace UE::PixelStreaming2
+
+#undef UE_API
diff --git a/Engine/Plugins/Media/PixelStreaming2/Source/PixelStreaming2/Internal/UtilsCore.h b/Engine/Plugins/Media/PixelStreaming2/Source/PixelStreaming2/Internal/UtilsCore.h
--- a/Engine/Plugins/Media/PixelStreaming2/Source/PixelStreaming2/Internal/UtilsCore.h
+++ b/Engine/Plugins/Media/PixelStreaming2/Source/PixelStreaming2/Internal/UtilsCore.h
@@ -25,7 +25,7 @@ namespace UE::PixelStreaming2
 		}
 	};
 
-	PIXELSTREAMING2_API inline bool IsStreamingSupported()
+	inline bool IsStreamingSupported()
 	{
 		// Pixel Streaming does not make sense without an RHI so we don't run in commandlets without one.
 		if (IsRunningCommandlet() && !IsAllowCommandletRendering())
@@ -36,7 +36,7 @@ namespace UE::PixelStreaming2
 		return true;
 	}
 
-	PIXELSTREAMING2_API inline bool IsPlatformSupported()
+	inline bool IsPlatformSupported()
 	{
 		bool bCompatible = true;
 
diff --git a/Engine/Plugins/Media/PixelStreaming2/Source/PixelStreaming2/Internal/VideoCapturer.h b/Engine/Plugins/Media/PixelStreaming2/Source/PixelStreaming2/Internal/VideoCapturer.h
--- a/Engine/Plugins/Media/PixelStreaming2/Source/PixelStreaming2/Internal/VideoCapturer.h
+++ b/Engine/Plugins/Media/PixelStreaming2/Source/PixelStreaming2/Internal/VideoCapturer.h
@@ -6,35 +6,38 @@
 #include "PixelCaptureCapturerMultiFormat.h"
 #include "VideoProducer.h"
 
+#define UE_API PIXELSTREAMING2_API
+
 namespace UE::PixelStreaming2
 {
 	/**
 	 * The start of PixelCapture pipeline. Frames enter the system when `OnFrameCaptured` is called.
 	 * This class creates the underlying PixelCapture `FPixelCaptureCapturer` that handles frame capture when `RequestFormat` is called.
 	 */
-	class PIXELSTREAMING2_API FVideoCapturer : public IPixelCaptureCapturerSource, public TSharedFromThis<FVideoCapturer>
+	class FVideoCapturer : public IPixelCaptureCapturerSource, public TSharedFromThis<FVideoCapturer>
 	{
 	public:
-		static TSharedPtr<FVideoCapturer> Create(TSharedPtr<FVideoProducer> VideoProducer = nullptr, bool bIsSRGB = false);
-		virtual ~FVideoCapturer();
+		static UE_API TSharedPtr<FVideoCapturer> Create(TSharedPtr<FVideoProducer> VideoProducer = nullptr, bool bIsSRGB = false);
+		UE_API virtual ~FVideoCapturer();
 
 		// Begin IPixelCaptureCapturerSource Interface
-		virtual TSharedPtr<FPixelCaptureCapturer> CreateCapturer(int32 FinalFormat, FIntPoint OutputResolution) override;
+		UE_API virtual TSharedPtr<FPixelCaptureCapturer> CreateCapturer(int32 FinalFormat, FIntPoint OutputResolution) override;
 		// End IPixelCaptureCapturerSource Interface
 
 		bool IsReady() const { return bReady; }
 
-		void SetVideoProducer(TSharedPtr<FVideoProducer> InVideoProducer);
+		UE_API void SetVideoProducer(TSharedPtr<FVideoProducer> InVideoProducer);
 
 		TSharedPtr<FVideoProducer> GetVideoProducer() { return VideoProducer; }
 
-		TSharedPtr<IPixelCaptureOutputFrame> WaitForFormat(int32 Format, TOptional<FIntPoint> Resolution = TOptional<FIntPoint>(), uint32 MaxWaitTime = 5000);
+		UE_API TSharedPtr<IPixelCaptureOutputFrame> WaitForFormat(int32 Format, TOptional<FIntPoint> Resolution = TOptional<FIntPoint>(), uint32 MaxWaitTime = 5000);
+
+		UE_API TSharedPtr<IPixelCaptureOutputFrame> RequestFormat(int32 Format, TOptional<FIntPoint> Resolution = TOptional<FIntPoint>());
 
-		TSharedPtr<IPixelCaptureOutputFrame> RequestFormat(int32 Format, TOptional<FIntPoint> Resolution = TOptional<FIntPoint>());
 
-		void ResetFrameCapturer();
+		UE_API void ResetFrameCapturer();
 
-		void OnFrame(const IPixelCaptureInputFrame& InputFrame);
+		UE_API void OnFrame(const IPixelCaptureInputFrame& InputFrame);
 
 		/**
 		 * This is broadcast each time a frame exits the adapt process. Used to synchronize framerates with input rates.
@@ -51,7 +54,7 @@ namespace UE::PixelStreaming2
 		FOnFrameCapturerCreated OnFrameCapturerCreated;
 
 	protected:
-		FVideoCapturer(TSharedPtr<FVideoProducer> VideoProducer, bool bIsSRGB);
+		UE_API FVideoCapturer(TSharedPtr<FVideoProducer> VideoProducer, bool bIsSRGB);
 
 		int32 LastFrameWidth = -1;
 		int32 LastFrameHeight = -1;
@@ -68,11 +71,13 @@ namespace UE::PixelStreaming2
 		FDelegateHandle CaptureUseFenceChangedHandle;
 		FDelegateHandle UseMediaCaptureChangedHandle;
 
-		void CreateFrameCapturer();
-		void OnSimulcastEnabledChanged(IConsoleVariable* Var);
-		void OnCaptureUseFenceChanged(IConsoleVariable* Var);
-		void OnUseMediaCaptureChanged(IConsoleVariable* Var);
-		void OnCaptureComplete();
+		UE_API void CreateFrameCapturer();
+		UE_API void OnSimulcastEnabledChanged(IConsoleVariable* Var);
+		UE_API void OnCaptureUseFenceChanged(IConsoleVariable* Var);
+		UE_API void OnUseMediaCaptureChanged(IConsoleVariable* Var);
+		UE_API void OnCaptureComplete();
 	};
 
 } // namespace UE::PixelStreaming2
+
+#undef UE_API
diff --git a/Engine/Plugins/Media/PixelStreaming2/Source/PixelStreaming2/Internal/VideoProducer.h b/Engine/Plugins/Media/PixelStreaming2/Source/PixelStreaming2/Internal/VideoProducer.h
--- a/Engine/Plugins/Media/PixelStreaming2/Source/PixelStreaming2/Internal/VideoProducer.h
+++ b/Engine/Plugins/Media/PixelStreaming2/Source/PixelStreaming2/Internal/VideoProducer.h
@@ -5,6 +5,8 @@
 #include "IPixelStreaming2VideoProducer.h"
 #include "MediaCapture.h"
 
+#define UE_API PIXELSTREAMING2_API
+
 namespace UE::PixelStreaming2
 {
 	namespace VideoProducerIdentifiers
@@ -16,7 +18,7 @@ namespace UE::PixelStreaming2
 		PIXELSTREAMING2_API extern const FString FVideoProducerRenderTarget;
 	} // namespace VideoProducerIdentifiers
 
-	class PIXELSTREAMING2_API FVideoProducerUserData : public FMediaCaptureUserData, public FPixelCaptureUserData
+	class FVideoProducerUserData : public FMediaCaptureUserData, public FPixelCaptureUserData
 	{
 	public:
 		FVideoProducerUserData() = default;
@@ -26,25 +28,27 @@ namespace UE::PixelStreaming2
 		FString ProducerName = TEXT("");
 	};
 
-	class PIXELSTREAMING2_API FVideoProducer : public IPixelStreaming2VideoProducer
+	class FVideoProducer : public IPixelStreaming2VideoProducer
 	{
 	public:
-		static TSharedPtr<FVideoProducer> Create();
+		static UE_API TSharedPtr<FVideoProducer> Create();
 		virtual ~FVideoProducer() = default;
 
 		DECLARE_EVENT_OneParam(FVideoProducer, FOnFramePushed, const IPixelCaptureInputFrame&);
 		FOnFramePushed OnFramePushed;
 
-		virtual void PushFrame(const IPixelCaptureInputFrame& InputFrame) override;
+		UE_API virtual void PushFrame(const IPixelCaptureInputFrame& InputFrame) override;
 
-		virtual FString ToString() override;
+		UE_API virtual FString ToString() override;
 
 		virtual bool IsFrameAlreadyCopied() { return false; }
 
 	protected:
 		FVideoProducer() = default;
 
-		virtual ETextureCreateFlags GetTexCreateFlags();
+		UE_API virtual ETextureCreateFlags GetTexCreateFlags();
 	};
 
-} // namespace UE::PixelStreaming2
\ No newline at end of file
+} // namespace UE::PixelStreaming2
+
+#undef UE_API
diff --git a/Engine/Plugins/Media/PixelStreaming2/Source/PixelStreaming2/Internal/VideoProducerBackBuffer.h b/Engine/Plugins/Media/PixelStreaming2/Source/PixelStreaming2/Internal/VideoProducerBackBuffer.h
--- a/Engine/Plugins/Media/PixelStreaming2/Source/PixelStreaming2/Internal/VideoProducerBackBuffer.h
+++ b/Engine/Plugins/Media/PixelStreaming2/Source/PixelStreaming2/Internal/VideoProducerBackBuffer.h
@@ -5,25 +5,29 @@
 #include "VideoProducer.h"
 #include "Widgets/SWindow.h"
 
+#define UE_API PIXELSTREAMING2_API
+
 namespace UE::PixelStreaming2
 {
 	/**
 	 * Use this if you want to send the contents of the UE backbuffer.
 	 */
-	class PIXELSTREAMING2_API FVideoProducerBackBuffer : public FVideoProducer
+	class FVideoProducerBackBuffer : public FVideoProducer
 	{
 	public:
-		static TSharedPtr<FVideoProducerBackBuffer> Create();
-		virtual ~FVideoProducerBackBuffer();
+		static UE_API TSharedPtr<FVideoProducerBackBuffer> Create();
+		UE_API virtual ~FVideoProducerBackBuffer();
 
-		virtual FString ToString() override;
+		UE_API virtual FString ToString() override;
 
 	private:
 		FVideoProducerBackBuffer() = default;
 
-		void OnBackBufferReady(SWindow& SlateWindow, const FTextureRHIRef& FrameBuffer);
+		UE_API void OnBackBufferReady(SWindow& SlateWindow, const FTextureRHIRef& FrameBuffer);
 
 		FDelegateHandle DelegateHandle;
 	};
 
-} // namespace UE::PixelStreaming2
\ No newline at end of file
+} // namespace UE::PixelStreaming2
+
+#undef UE_API
diff --git a/Engine/Plugins/Media/PixelStreaming2/Source/PixelStreaming2/Internal/VideoProducerMediaCapture.h b/Engine/Plugins/Media/PixelStreaming2/Source/PixelStreaming2/Internal/VideoProducerMediaCapture.h
--- a/Engine/Plugins/Media/PixelStreaming2/Source/PixelStreaming2/Internal/VideoProducerMediaCapture.h
+++ b/Engine/Plugins/Media/PixelStreaming2/Source/PixelStreaming2/Internal/VideoProducerMediaCapture.h
@@ -8,6 +8,8 @@
 
 #include "VideoProducerMediaCapture.generated.h"
 
+#define UE_API PIXELSTREAMING2_API
+
 UCLASS(BlueprintType)
 class UPixelStreaming2MediaIOOutput : public UMediaOutput
 {
@@ -29,14 +31,14 @@ namespace UE::PixelStreaming2
 	/**
 	 * Use this if you want to send media capture frames as video input.
 	 */
-	class PIXELSTREAMING2_API FVideoProducerMediaCapture : public FVideoProducer, public TSharedFromThis<FVideoProducerMediaCapture>
+	class FVideoProducerMediaCapture : public FVideoProducer, public TSharedFromThis<FVideoProducerMediaCapture>
 	{
 	public:
 		/**
 		 * @brief Creates a MediaIO capture of the active viewport and starts capturing as soon as possible.
 		 * @return A video input backed by the created MediaIO capture that sends frames from the active viewport.
 		 */
-		static TSharedPtr<FVideoProducerMediaCapture> CreateActiveViewportCapture();
+		static UE_API TSharedPtr<FVideoProducerMediaCapture> CreateActiveViewportCapture();
 
 		/**
 		 * @brief Creates a video input where the user can specify their own MediaIO output and capture.
@@ -45,23 +47,23 @@ namespace UE::PixelStreaming2
 		 * @param MediaCapture The custom MediaIOCapture that will pass its captured frames as video input.
 		 * @return A video input backed by the passed in MediaIO capture.
 		 */
-		static TSharedPtr<FVideoProducerMediaCapture> Create(TObjectPtr<UPixelStreaming2MediaIOCapture> MediaCapture);
+		static UE_API TSharedPtr<FVideoProducerMediaCapture> Create(TObjectPtr<UPixelStreaming2MediaIOCapture> MediaCapture);
 
-		virtual ~FVideoProducerMediaCapture();
+		UE_API virtual ~FVideoProducerMediaCapture();
 
-		virtual FString ToString() override;
+		UE_API virtual FString ToString() override;
 
 		virtual bool IsFrameAlreadyCopied() override { return true; }
 
 	protected:
-		void StartActiveViewportCapture();
-		void LateStartActiveViewportCapture();
+		UE_API void StartActiveViewportCapture();
+		UE_API void LateStartActiveViewportCapture();
 
 	private:
 		FVideoProducerMediaCapture() = default;
-		FVideoProducerMediaCapture(TObjectPtr<UPixelStreaming2MediaIOCapture> MediaCapture);
+		UE_API FVideoProducerMediaCapture(TObjectPtr<UPixelStreaming2MediaIOCapture> MediaCapture);
 
-		void OnCaptureActiveViewportStateChanged();
+		UE_API void OnCaptureActiveViewportStateChanged();
 
 	private:
 		TObjectPtr<UPixelStreaming2MediaIOCapture> MediaCapture = nullptr;
@@ -70,3 +72,5 @@ namespace UE::PixelStreaming2
 	};
 
 } // namespace UE::PixelStreaming2
+
+#undef UE_API
diff --git a/Engine/Plugins/Media/PixelStreaming2/Source/PixelStreaming2/Internal/VideoProducerPIEViewport.h b/Engine/Plugins/Media/PixelStreaming2/Source/PixelStreaming2/Internal/VideoProducerPIEViewport.h
--- a/Engine/Plugins/Media/PixelStreaming2/Source/PixelStreaming2/Internal/VideoProducerPIEViewport.h
+++ b/Engine/Plugins/Media/PixelStreaming2/Source/PixelStreaming2/Internal/VideoProducerPIEViewport.h
@@ -5,25 +5,29 @@
 #include "RHIFwd.h"
 #include "VideoProducerViewportBase.h"
 
+#define UE_API PIXELSTREAMING2_API
+
 namespace UE::PixelStreaming2
 {
 	/**
 	 * An extension of the back buffer input that can handle PIE sessions. Primarily to be used in blueprints
 	 */
-	class PIXELSTREAMING2_API FVideoProducerPIEViewport : public FVideoProducerViewportBase
+	class FVideoProducerPIEViewport : public FVideoProducerViewportBase
 	{
 	public:
-		static TSharedPtr<FVideoProducerPIEViewport> Create();
+		static UE_API TSharedPtr<FVideoProducerPIEViewport> Create();
 		virtual ~FVideoProducerPIEViewport() = default;
 
-		virtual FString ToString() override;
+		UE_API virtual FString ToString() override;
 
 	protected:
-		virtual bool ShouldCaptureViewport() override;
+		UE_API virtual bool ShouldCaptureViewport() override;
 
 	private:
 		FVideoProducerPIEViewport() = default;
  
-		void OnPreTick(float DeltaTime);
+		UE_API void OnPreTick(float DeltaTime);
 	};
-} // namespace UE::PixelStreaming2
\ No newline at end of file
+} // namespace UE::PixelStreaming2
+
+#undef UE_API
diff --git a/Engine/Plugins/Media/PixelStreaming2/Source/PixelStreaming2/Internal/VideoProducerRenderTarget.h b/Engine/Plugins/Media/PixelStreaming2/Source/PixelStreaming2/Internal/VideoProducerRenderTarget.h
--- a/Engine/Plugins/Media/PixelStreaming2/Source/PixelStreaming2/Internal/VideoProducerRenderTarget.h
+++ b/Engine/Plugins/Media/PixelStreaming2/Source/PixelStreaming2/Internal/VideoProducerRenderTarget.h
@@ -4,6 +4,8 @@
 
 #include "VideoProducer.h"
 
+#define UE_API PIXELSTREAMING2_API
+
 class UTextureRenderTarget2D;
 
 namespace UE::PixelStreaming2
@@ -11,20 +13,22 @@ namespace UE::PixelStreaming2
 	/**
 	 * Use this if you want to send the contents of a render target.
 	 */
-	class PIXELSTREAMING2_API FVideoProducerRenderTarget : public FVideoProducer
+	class FVideoProducerRenderTarget : public FVideoProducer
 	{
 	public:
-		static TSharedPtr<FVideoProducerRenderTarget> Create(UTextureRenderTarget2D* Target);
-		virtual ~FVideoProducerRenderTarget();
+		static UE_API TSharedPtr<FVideoProducerRenderTarget> Create(UTextureRenderTarget2D* Target);
+		UE_API virtual ~FVideoProducerRenderTarget();
 
-		virtual FString ToString() override;
+		UE_API virtual FString ToString() override;
 
 	private:
-		FVideoProducerRenderTarget(UTextureRenderTarget2D* InTarget);
-		void OnEndFrameRenderThread();
+		UE_API FVideoProducerRenderTarget(UTextureRenderTarget2D* InTarget);
+		UE_API void OnEndFrameRenderThread();
 
 		UTextureRenderTarget2D* Target = nullptr;
 		FDelegateHandle			DelegateHandle;
 	};
 
-} // namespace UE::PixelStreaming2
\ No newline at end of file
+} // namespace UE::PixelStreaming2
+
+#undef UE_API
diff --git a/Engine/Plugins/Media/PixelStreaming2/Source/PixelStreaming2/Internal/VideoProducerViewportBase.h b/Engine/Plugins/Media/PixelStreaming2/Source/PixelStreaming2/Internal/VideoProducerViewportBase.h
--- a/Engine/Plugins/Media/PixelStreaming2/Source/PixelStreaming2/Internal/VideoProducerViewportBase.h
+++ b/Engine/Plugins/Media/PixelStreaming2/Source/PixelStreaming2/Internal/VideoProducerViewportBase.h
@@ -7,6 +7,8 @@
 #include "Widgets/SViewport.h"
 #include "Widgets/SWindow.h"
 
+#define UE_API PIXELSTREAMING2_API
+
 namespace UE::PixelStreaming2
 {
 	/**
@@ -16,21 +18,21 @@ namespace UE::PixelStreaming2
 	 * Implementations MUST set the TargetViewport and TargetWidget and then call CalculateCaptureRegion at least once to configure the captured region.
 	 * Additionally, implementations MUST override ShouldCaptureViewport with an optional check to prevent OnBackBufferReadyToPresent from sending frames.
 	 */
-	class PIXELSTREAMING2_API FVideoProducerViewportBase : public FVideoProducer
+	class FVideoProducerViewportBase : public FVideoProducer
 	{
 	public:
 		virtual ~FVideoProducerViewportBase() = default;
 
-		virtual bool IsFrameAlreadyCopied() override;
+		UE_API virtual bool IsFrameAlreadyCopied() override;
 
     protected:
         FVideoProducerViewportBase() = default;
 
         // Viewport specific implementations should override this method with their own check
-        virtual bool ShouldCaptureViewport();
+        UE_API virtual bool ShouldCaptureViewport();
 
-        void CalculateCaptureRegion(TSharedRef<SViewport> Viewport, TSharedRef<SWindow> Window);
-		void OnBackBufferReadyToPresent(SWindow& Window, const FTextureRHIRef& BackBuffer);
+        UE_API void CalculateCaptureRegion(TSharedRef<SViewport> Viewport, TSharedRef<SWindow> Window);
+		UE_API void OnBackBufferReadyToPresent(SWindow& Window, const FTextureRHIRef& BackBuffer);
 
         FIntRect CaptureRect;
 		FCriticalSection CaptureRectCS;
@@ -38,4 +40,6 @@ namespace UE::PixelStreaming2
         TWeakPtr<SViewport> TargetViewport;
 		TWeakPtr<SWindow> TargetWindow;
 	};
-} // namespace UE::PixelStreaming2
\ No newline at end of file
+} // namespace UE::PixelStreaming2
+
+#undef UE_API
diff --git a/Engine/Plugins/Media/PixelStreaming2/Source/PixelStreaming2/Internal/VideoSink.h b/Engine/Plugins/Media/PixelStreaming2/Source/PixelStreaming2/Internal/VideoSink.h
--- a/Engine/Plugins/Media/PixelStreaming2/Source/PixelStreaming2/Internal/VideoSink.h
+++ b/Engine/Plugins/Media/PixelStreaming2/Source/PixelStreaming2/Internal/VideoSink.h
@@ -8,21 +8,23 @@
 #include "IPixelStreaming2VideoSink.h"
 #include "Rendering/SlateRenderer.h"
 
+#define UE_API PIXELSTREAMING2_API
+
 namespace UE::PixelStreaming2
 {
-	class PIXELSTREAMING2_API FVideoSink : public IPixelStreaming2VideoSink
+	class FVideoSink : public IPixelStreaming2VideoSink
 	{
 	public:
-		virtual ~FVideoSink();
+		UE_API virtual ~FVideoSink();
 
-		virtual void AddVideoConsumer(const TWeakPtrVariant<IPixelStreaming2VideoConsumer>& VideoConsumer) override;
-		virtual void RemoveVideoConsumer(const TWeakPtrVariant<IPixelStreaming2VideoConsumer>& VideoConsumer) override;
+		UE_API virtual void AddVideoConsumer(const TWeakPtrVariant<IPixelStreaming2VideoConsumer>& VideoConsumer) override;
+		UE_API virtual void RemoveVideoConsumer(const TWeakPtrVariant<IPixelStreaming2VideoConsumer>& VideoConsumer) override;
 
-		bool HasVideoConsumers();
+		UE_API bool HasVideoConsumers();
 
-		void SetMuted(bool bIsMuted);
+		UE_API void SetMuted(bool bIsMuted);
 
-		void OnVideoData(FTextureRHIRef Frame);
+		UE_API void OnVideoData(FTextureRHIRef Frame);
 
 	protected:
 		FCriticalSection									   VideoConsumersCS;
@@ -31,4 +33,6 @@ namespace UE::PixelStreaming2
 		bool bIsMuted = false;
 	};
 
-} // namespace UE::PixelStreaming2
\ No newline at end of file
+} // namespace UE::PixelStreaming2
+
+#undef UE_API
diff --git a/Engine/Plugins/Media/PixelStreaming2/Source/PixelStreaming2/Internal/VideoSource.h b/Engine/Plugins/Media/PixelStreaming2/Source/PixelStreaming2/Internal/VideoSource.h
--- a/Engine/Plugins/Media/PixelStreaming2/Source/PixelStreaming2/Internal/VideoSource.h
+++ b/Engine/Plugins/Media/PixelStreaming2/Source/PixelStreaming2/Internal/VideoSource.h
@@ -5,18 +5,22 @@
 
 #include "VideoCapturer.h"
 
+#define UE_API PIXELSTREAMING2_API
+
 namespace UE::PixelStreaming2
 {
-	class PIXELSTREAMING2_API FVideoSource
+	class FVideoSource
 	{
 	public:
 		virtual ~FVideoSource() = default;
 
-		virtual void PushFrame();
-		virtual void ForceKeyFrame();
-		virtual void SetMuted(bool bIsMuted);
+		UE_API virtual void PushFrame();
+		UE_API virtual void ForceKeyFrame();
+		UE_API virtual void SetMuted(bool bIsMuted);
 
 	protected:
 		bool bIsMuted = false;
 	};
 } // namespace UE::PixelStreaming2
+
+#undef UE_API
diff --git a/Engine/Plugins/Media/PixelStreaming2/Source/PixelStreaming2/Internal/VideoSourceGroup.h b/Engine/Plugins/Media/PixelStreaming2/Source/PixelStreaming2/Internal/VideoSourceGroup.h
--- a/Engine/Plugins/Media/PixelStreaming2/Source/PixelStreaming2/Internal/VideoSourceGroup.h
+++ b/Engine/Plugins/Media/PixelStreaming2/Source/PixelStreaming2/Internal/VideoSourceGroup.h
@@ -4,42 +4,44 @@
 
 #include "VideoCapturer.h"
 
+#define UE_API PIXELSTREAMING2_API
+
 namespace UE::PixelStreaming2
 {
 	class FVideoSource;
 	class FFrameRunnable;
 
-	class PIXELSTREAMING2_API FVideoSourceGroup : public TSharedFromThis<FVideoSourceGroup>
+	class FVideoSourceGroup : public TSharedFromThis<FVideoSourceGroup>
 	{
 	public:
-		static TSharedPtr<FVideoSourceGroup> Create(TSharedPtr<FVideoCapturer> InVideoCapturer);
-		~FVideoSourceGroup();
+		static UE_API TSharedPtr<FVideoSourceGroup> Create(TSharedPtr<FVideoCapturer> InVideoCapturer);
+		UE_API ~FVideoSourceGroup();
 
-		void  SetFPS(int32 InFramesPerSecond);
-		int32 GetFPS();
+		UE_API void  SetFPS(int32 InFramesPerSecond);
+		UE_API int32 GetFPS();
 
-		void SetDecoupleFramerate(bool bDecouple);
+		UE_API void SetDecoupleFramerate(bool bDecouple);
 
-		void AddVideoSource(TSharedPtr<FVideoSource> VideoSource);
-		void RemoveVideoSource(const FVideoSource* ToRemove);
-		void RemoveAllVideoSources();
+		UE_API void AddVideoSource(TSharedPtr<FVideoSource> VideoSource);
+		UE_API void RemoveVideoSource(const FVideoSource* ToRemove);
+		UE_API void RemoveAllVideoSources();
 
-		void Start();
-		void Stop();
-		void PushFrame();
+		UE_API void Start();
+		UE_API void Stop();
+		UE_API void PushFrame();
 
-		void ForceKeyFrame();
+		UE_API void ForceKeyFrame();
 
 	private:
-		FVideoSourceGroup();
+		UE_API FVideoSourceGroup();
 
-		void StartThread();
-		void StopThread();
+		UE_API void StartThread();
+		UE_API void StopThread();
 
-		void OnFrameCaptured();
+		UE_API void OnFrameCaptured();
 
-		void OnWebRtcFpsChanged(IConsoleVariable* Var);
-		void OnDecoupleFramerateChanged(IConsoleVariable* Var);
+		UE_API void OnWebRtcFpsChanged(IConsoleVariable* Var);
+		UE_API void OnDecoupleFramerateChanged(IConsoleVariable* Var);
 
 		bool bFPSOverride = false;
 		bool bDecoupleOverride = false;
@@ -61,3 +63,5 @@ namespace UE::PixelStreaming2
 		mutable FCriticalSection RunnableCriticalSection;
 	};
 } // namespace UE::PixelStreaming2
+
+#undef UE_API
diff --git a/Engine/Plugins/Media/PixelStreaming2/Source/PixelStreaming2/Private/Blueprints/PixelStreaming2AudioComponent.cpp b/Engine/Plugins/Media/PixelStreaming2/Source/PixelStreaming2/Private/Blueprints/PixelStreaming2AudioComponent.cpp
--- a/Engine/Plugins/Media/PixelStreaming2/Source/PixelStreaming2/Private/Blueprints/PixelStreaming2AudioComponent.cpp
+++ b/Engine/Plugins/Media/PixelStreaming2/Source/PixelStreaming2/Private/Blueprints/PixelStreaming2AudioComponent.cpp
@@ -9,6 +9,8 @@
 #include "SampleBuffer.h"
 #include "SoundGenerator.h"
 
+#include UE_INLINE_GENERATED_CPP_BY_NAME(PixelStreaming2AudioComponent)
+
 /**
  * Component that recieves audio from a remote webrtc connection and outputs it into UE using a "synth component".
  */
diff --git a/Engine/Plugins/Media/PixelStreaming2/Source/PixelStreaming2/Private/Blueprints/PixelStreaming2Blueprints.cpp b/Engine/Plugins/Media/PixelStreaming2/Source/PixelStreaming2/Private/Blueprints/PixelStreaming2Blueprints.cpp
--- a/Engine/Plugins/Media/PixelStreaming2/Source/PixelStreaming2/Private/Blueprints/PixelStreaming2Blueprints.cpp
+++ b/Engine/Plugins/Media/PixelStreaming2/Source/PixelStreaming2/Private/Blueprints/PixelStreaming2Blueprints.cpp
@@ -8,6 +8,8 @@
 #include "PixelStreaming2AudioComponent.h"
 #include "PixelStreaming2Utils.h"
 
+#include UE_INLINE_GENERATED_CPP_BY_NAME(PixelStreaming2Blueprints)
+
 TSharedPtr<IPixelStreaming2Streamer> GetStreamer(FString StreamerId)
 {
 	TSharedPtr<IPixelStreaming2Streamer> Streamer = IPixelStreaming2Module::Get().FindStreamer(StreamerId);
@@ -160,4 +162,4 @@ void UPixelStreaming2Blueprints::AddJsonStringValue(const FString& Descriptor, F
 UPixelStreaming2Delegates* UPixelStreaming2Blueprints::GetDelegates()
 {
 	return UPixelStreaming2Delegates::Get();
-}
\ No newline at end of file
+}
diff --git a/Engine/Plugins/Media/PixelStreaming2/Source/PixelStreaming2/Private/Blueprints/PixelStreaming2InputComponent.cpp b/Engine/Plugins/Media/PixelStreaming2/Source/PixelStreaming2/Private/Blueprints/PixelStreaming2InputComponent.cpp
--- a/Engine/Plugins/Media/PixelStreaming2/Source/PixelStreaming2/Private/Blueprints/PixelStreaming2InputComponent.cpp
+++ b/Engine/Plugins/Media/PixelStreaming2/Source/PixelStreaming2/Private/Blueprints/PixelStreaming2InputComponent.cpp
@@ -11,6 +11,8 @@
 #include "PixelStreaming2Utils.h"
 #include "PixelStreaming2InputEnums.h"
 
+#include UE_INLINE_GENERATED_CPP_BY_NAME(PixelStreaming2InputComponent)
+
 UPixelStreaming2Input::UPixelStreaming2Input(const FObjectInitializer& ObjectInitializer)
 	: Super(ObjectInitializer)
 {
diff --git a/Engine/Plugins/Media/PixelStreaming2/Source/PixelStreaming2/Private/Blueprints/PixelStreaming2MediaTexture.cpp b/Engine/Plugins/Media/PixelStreaming2/Source/PixelStreaming2/Private/Blueprints/PixelStreaming2MediaTexture.cpp
--- a/Engine/Plugins/Media/PixelStreaming2/Source/PixelStreaming2/Private/Blueprints/PixelStreaming2MediaTexture.cpp
+++ b/Engine/Plugins/Media/PixelStreaming2/Source/PixelStreaming2/Private/Blueprints/PixelStreaming2MediaTexture.cpp
@@ -8,6 +8,8 @@
 #include "RenderUtils.h"
 #include "TextureResource.h"
 
+#include UE_INLINE_GENERATED_CPP_BY_NAME(PixelStreaming2MediaTexture)
+
 constexpr int32 DEFAULT_WIDTH = 1920;
 constexpr int32 DEFAULT_HEIGHT = 1080;
 
diff --git a/Engine/Plugins/Media/PixelStreaming2/Source/PixelStreaming2/Private/Blueprints/PixelStreaming2StreamerComponent.cpp b/Engine/Plugins/Media/PixelStreaming2/Source/PixelStreaming2/Private/Blueprints/PixelStreaming2StreamerComponent.cpp
--- a/Engine/Plugins/Media/PixelStreaming2/Source/PixelStreaming2/Private/Blueprints/PixelStreaming2StreamerComponent.cpp
+++ b/Engine/Plugins/Media/PixelStreaming2/Source/PixelStreaming2/Private/Blueprints/PixelStreaming2StreamerComponent.cpp
@@ -8,6 +8,8 @@
 #include "Engine/GameEngine.h"
 #include "Slate/SceneViewport.h"
 
+#include UE_INLINE_GENERATED_CPP_BY_NAME(PixelStreaming2StreamerComponent)
+
 UPixelStreaming2StreamerComponent::UPixelStreaming2StreamerComponent(const FObjectInitializer& ObjectInitializer)
 	: Super(ObjectInitializer)
 {
diff --git a/Engine/Plugins/Media/PixelStreaming2/Source/PixelStreaming2/Private/Blueprints/PixelStreaming2VideoComponent.cpp b/Engine/Plugins/Media/PixelStreaming2/Source/PixelStreaming2/Private/Blueprints/PixelStreaming2VideoComponent.cpp
--- a/Engine/Plugins/Media/PixelStreaming2/Source/PixelStreaming2/Private/Blueprints/PixelStreaming2VideoComponent.cpp
+++ b/Engine/Plugins/Media/PixelStreaming2/Source/PixelStreaming2/Private/Blueprints/PixelStreaming2VideoComponent.cpp
@@ -7,6 +7,8 @@
 #include "IPixelStreaming2Streamer.h"
 #include "Logging.h"
 
+#include UE_INLINE_GENERATED_CPP_BY_NAME(PixelStreaming2VideoComponent)
+
 /**
  * Component that recieves video from a remote webrtc connection and outputs it into UE using a "synth component".
  */
diff --git a/Engine/Plugins/Media/PixelStreaming2/Source/PixelStreaming2/Private/Blueprints/PixelStreaming2VideoProducer.cpp b/Engine/Plugins/Media/PixelStreaming2/Source/PixelStreaming2/Private/Blueprints/PixelStreaming2VideoProducer.cpp
--- a/Engine/Plugins/Media/PixelStreaming2/Source/PixelStreaming2/Private/Blueprints/PixelStreaming2VideoProducer.cpp
+++ b/Engine/Plugins/Media/PixelStreaming2/Source/PixelStreaming2/Private/Blueprints/PixelStreaming2VideoProducer.cpp
@@ -7,6 +7,8 @@
 #include "VideoProducerPIEViewport.h"
 #include "VideoProducerRenderTarget.h"
 
+#include UE_INLINE_GENERATED_CPP_BY_NAME(PixelStreaming2VideoProducer)
+
 /**
  * ---------- UPixelStreaming2VideoProducerBase -------------------
  */
diff --git a/Engine/Plugins/Media/PixelStreaming2/Source/PixelStreaming2/Private/PixelStreaming2Delegates.cpp b/Engine/Plugins/Media/PixelStreaming2/Source/PixelStreaming2/Private/PixelStreaming2Delegates.cpp
--- a/Engine/Plugins/Media/PixelStreaming2/Source/PixelStreaming2/Private/PixelStreaming2Delegates.cpp
+++ b/Engine/Plugins/Media/PixelStreaming2/Source/PixelStreaming2/Private/PixelStreaming2Delegates.cpp
@@ -2,6 +2,8 @@
 
 #include "PixelStreaming2Delegates.h"
 
+#include UE_INLINE_GENERATED_CPP_BY_NAME(PixelStreaming2Delegates)
+
 UPixelStreaming2Delegates* UPixelStreaming2Delegates::Singleton = nullptr;
 
 UPixelStreaming2Delegates* UPixelStreaming2Delegates::Get()
diff --git a/Engine/Plugins/Media/PixelStreaming2/Source/PixelStreaming2/Private/PixelStreaming2MediaIOCapture.cpp b/Engine/Plugins/Media/PixelStreaming2/Source/PixelStreaming2/Private/PixelStreaming2MediaIOCapture.cpp
--- a/Engine/Plugins/Media/PixelStreaming2/Source/PixelStreaming2/Private/PixelStreaming2MediaIOCapture.cpp
+++ b/Engine/Plugins/Media/PixelStreaming2/Source/PixelStreaming2/Private/PixelStreaming2MediaIOCapture.cpp
@@ -7,6 +7,8 @@
 #include "PixelCaptureInputFrameRHI.h"
 #include "ScreenPass.h"
 
+#include UE_INLINE_GENERATED_CPP_BY_NAME(PixelStreaming2MediaIOCapture)
+
 void UPixelStreaming2MediaIOCapture::OnRHIResourceCaptured_RenderingThread(
 	FRHICommandListImmediate& /* RHICmdList */,
 	const FCaptureBaseData&								   InBaseData,
diff --git a/Engine/Plugins/Media/PixelStreaming2/Source/PixelStreaming2/Private/VideoProducerMediaCapture.cpp b/Engine/Plugins/Media/PixelStreaming2/Source/PixelStreaming2/Private/VideoProducerMediaCapture.cpp
--- a/Engine/Plugins/Media/PixelStreaming2/Source/PixelStreaming2/Private/VideoProducerMediaCapture.cpp
+++ b/Engine/Plugins/Media/PixelStreaming2/Source/PixelStreaming2/Private/VideoProducerMediaCapture.cpp
@@ -4,6 +4,8 @@
 
 #include "Logging.h"
 
+#include UE_INLINE_GENERATED_CPP_BY_NAME(VideoProducerMediaCapture)
+
 namespace UE::PixelStreaming2
 {
 
diff --git a/Engine/Plugins/Media/PixelStreaming2/Source/PixelStreaming2/Public/IPixelStreaming2Module.h b/Engine/Plugins/Media/PixelStreaming2/Source/PixelStreaming2/Public/IPixelStreaming2Module.h
--- a/Engine/Plugins/Media/PixelStreaming2/Source/PixelStreaming2/Public/IPixelStreaming2Module.h
+++ b/Engine/Plugins/Media/PixelStreaming2/Source/PixelStreaming2/Public/IPixelStreaming2Module.h
@@ -16,7 +16,7 @@
  * This class provides access to streamers, video and audio producers, and handles starting and stopping the streaming process.
  * It also allows interaction with the signalling server and manages the lifecycle of streamers within the Pixel Streaming system.
  */
-class PIXELSTREAMING2_API IPixelStreaming2Module : public IModuleInterface
+class IPixelStreaming2Module : public IModuleInterface
 {
 public:
 	/**
@@ -139,4 +139,4 @@ public:
 	 * @param Func The lambda to execute with each streamer.
 	 */
 	virtual void ForEachStreamer(const TFunction<void(TSharedPtr<IPixelStreaming2Streamer>)>& Func) = 0;
-};
\ No newline at end of file
+};
diff --git a/Engine/Plugins/Media/PixelStreaming2/Source/PixelStreaming2/Public/PixelStreaming2Delegates.h b/Engine/Plugins/Media/PixelStreaming2/Source/PixelStreaming2/Public/PixelStreaming2Delegates.h
--- a/Engine/Plugins/Media/PixelStreaming2/Source/PixelStreaming2/Public/PixelStreaming2Delegates.h
+++ b/Engine/Plugins/Media/PixelStreaming2/Source/PixelStreaming2/Public/PixelStreaming2Delegates.h
@@ -8,12 +8,14 @@
 
 #include "PixelStreaming2Delegates.generated.h"
 
+#define UE_API PIXELSTREAMING2_API
+
 /**
  * Pixel Streaming Delegates that can be invoked when pixel streaming events take place.
  * Includes blueprint and native c++ delegates.
  */
-UCLASS()
-class PIXELSTREAMING2_API UPixelStreaming2Delegates : public UObject
+UCLASS(MinimalAPI)
+class UPixelStreaming2Delegates : public UObject
 {
 	GENERATED_BODY()
 
@@ -163,11 +165,13 @@ public:
 	 * @param Get the UPixelStreaming2Delegates singleton.
 	 * @return UPixelStreaming2Delegates pointer.
 	 */
-	static UPixelStreaming2Delegates* Get();
+	static UE_API UPixelStreaming2Delegates* Get();
 
-	virtual ~UPixelStreaming2Delegates();
+	UE_API virtual ~UPixelStreaming2Delegates();
 
 private:
 	// The singleton object.
-	static UPixelStreaming2Delegates* Singleton;
+	static UE_API UPixelStreaming2Delegates* Singleton;
 };
+
+#undef UE_API
diff --git a/Engine/Plugins/Media/PixelStreaming2/Source/PixelStreaming2Core/Public/IPixelStreaming2AudioConsumer.h b/Engine/Plugins/Media/PixelStreaming2/Source/PixelStreaming2Core/Public/IPixelStreaming2AudioConsumer.h
--- a/Engine/Plugins/Media/PixelStreaming2/Source/PixelStreaming2Core/Public/IPixelStreaming2AudioConsumer.h
+++ b/Engine/Plugins/Media/PixelStreaming2/Source/PixelStreaming2Core/Public/IPixelStreaming2AudioConsumer.h
@@ -16,7 +16,7 @@ class UPixelStreaming2AudioConsumer : public UInterface
  * An "Audio Consumer" is an object that is responsible for outputting the audio received from a peer. For example, by
  * passing the audio into a UE submix.
  */
-class PIXELSTREAMING2CORE_API IPixelStreaming2AudioConsumer
+class IPixelStreaming2AudioConsumer
 {
 	GENERATED_BODY()
 
diff --git a/Engine/Plugins/Media/PixelStreaming2/Source/PixelStreaming2Core/Public/IPixelStreaming2AudioProducer.h b/Engine/Plugins/Media/PixelStreaming2/Source/PixelStreaming2Core/Public/IPixelStreaming2AudioProducer.h
--- a/Engine/Plugins/Media/PixelStreaming2/Source/PixelStreaming2Core/Public/IPixelStreaming2AudioProducer.h
+++ b/Engine/Plugins/Media/PixelStreaming2/Source/PixelStreaming2Core/Public/IPixelStreaming2AudioProducer.h
@@ -21,7 +21,7 @@ class UPixelStreaming2AudioProducer : public UInterface
  *
  * Note: Any audio you push in with `PushAudio` will be mixed with other audio producers before being streamed
  */
-class PIXELSTREAMING2CORE_API IPixelStreaming2AudioProducer
+class IPixelStreaming2AudioProducer
 {
 	GENERATED_BODY()
 
diff --git a/Engine/Plugins/Media/PixelStreaming2/Source/PixelStreaming2Core/Public/IPixelStreaming2AudioSink.h b/Engine/Plugins/Media/PixelStreaming2/Source/PixelStreaming2Core/Public/IPixelStreaming2AudioSink.h
--- a/Engine/Plugins/Media/PixelStreaming2/Source/PixelStreaming2Core/Public/IPixelStreaming2AudioSink.h
+++ b/Engine/Plugins/Media/PixelStreaming2/Source/PixelStreaming2Core/Public/IPixelStreaming2AudioSink.h
@@ -18,7 +18,7 @@ class UPixelStreaming2AudioSink : public UInterface
  * An "Audio Sink" is an object that receives audio from a singular peer. From here, you can add multiple consumers to
  * output the received audio.
  */
-class PIXELSTREAMING2CORE_API IPixelStreaming2AudioSink
+class IPixelStreaming2AudioSink
 {
 	GENERATED_BODY()
 
@@ -40,4 +40,4 @@ public:
 	virtual void RemoveAudioConsumer(IPixelStreaming2AudioConsumer* AudioConsumer) {}
 
 	virtual void RemoveAudioConsumer(const TWeakPtrVariant<IPixelStreaming2AudioConsumer>& AudioConsumer) = 0;
-};
\ No newline at end of file
+};
diff --git a/Engine/Plugins/Media/PixelStreaming2/Source/PixelStreaming2Core/Public/IPixelStreaming2Streamer.h b/Engine/Plugins/Media/PixelStreaming2/Source/PixelStreaming2Core/Public/IPixelStreaming2Streamer.h
--- a/Engine/Plugins/Media/PixelStreaming2/Source/PixelStreaming2Core/Public/IPixelStreaming2Streamer.h
+++ b/Engine/Plugins/Media/PixelStreaming2/Source/PixelStreaming2Core/Public/IPixelStreaming2Streamer.h
@@ -8,6 +8,8 @@
 #include "IPixelStreaming2VideoSink.h"
 #include "Templates/SharedPointer.h"
 
+#define UE_API PIXELSTREAMING2CORE_API
+
 class UTexture2D;
 namespace UE::PixelStreaming2
 {
@@ -19,7 +21,7 @@ namespace UE::PixelStreaming2
  * This class manages the core operations of streaming, such as setting the framerate, handling video input,
  * and interacting with the signalling server. It also allows for control over connected players and sending data to them.
  */
-class PIXELSTREAMING2CORE_API IPixelStreaming2Streamer
+class IPixelStreaming2Streamer
 {
 public:
 	virtual ~IPixelStreaming2Streamer() = default;
@@ -269,24 +271,26 @@ public:
  * The IPixelStreaming2StreamerFactory interface provides functionality for creating custom streamers based on an identifier.
  *
  */
-class PIXELSTREAMING2CORE_API IPixelStreaming2StreamerFactory : public IModularFeature
+class IPixelStreaming2StreamerFactory : public IModularFeature
 {
 public:
-	IPixelStreaming2StreamerFactory();
-	virtual ~IPixelStreaming2StreamerFactory();
+	UE_API IPixelStreaming2StreamerFactory();
+	UE_API virtual ~IPixelStreaming2StreamerFactory();
 
 	virtual FString GetStreamType() = 0;
 
-	static void RegisterStreamerFactory(IPixelStreaming2StreamerFactory* InFactory);
+	static UE_API void RegisterStreamerFactory(IPixelStreaming2StreamerFactory* InFactory);
 
-	static void UnregisterStreamerFactory(IPixelStreaming2StreamerFactory* InFactory);
+	static UE_API void UnregisterStreamerFactory(IPixelStreaming2StreamerFactory* InFactory);
 
-	static IPixelStreaming2StreamerFactory* Get(const FString& InType);
+	static UE_API IPixelStreaming2StreamerFactory* Get(const FString& InType);
 
-	static TArray<FString> GetAvailableFactoryTypes();
+	static UE_API TArray<FString> GetAvailableFactoryTypes();
 
 private:
 	// Only the module should call CreateNewStreamer on the factories.
 	friend UE::PixelStreaming2::FPixelStreaming2Module;
 	virtual TSharedPtr<IPixelStreaming2Streamer> CreateNewStreamer(const FString& StreamerId) = 0;
 };
+
+#undef UE_API
diff --git a/Engine/Plugins/Media/PixelStreaming2/Source/PixelStreaming2Core/Public/IPixelStreaming2VideoConsumer.h b/Engine/Plugins/Media/PixelStreaming2/Source/PixelStreaming2Core/Public/IPixelStreaming2VideoConsumer.h
--- a/Engine/Plugins/Media/PixelStreaming2/Source/PixelStreaming2Core/Public/IPixelStreaming2VideoConsumer.h
+++ b/Engine/Plugins/Media/PixelStreaming2/Source/PixelStreaming2Core/Public/IPixelStreaming2VideoConsumer.h
@@ -17,7 +17,7 @@ class UPixelStreaming2VideoConsumer : public UInterface
  * An "Video Consumer" is an object that is responsible for outputting the video received from a peer. For example, by
  * rendering to a render target.
  */
-class PIXELSTREAMING2CORE_API IPixelStreaming2VideoConsumer
+class IPixelStreaming2VideoConsumer
 {
 	GENERATED_BODY()
 
diff --git a/Engine/Plugins/Media/PixelStreaming2/Source/PixelStreaming2Core/Public/IPixelStreaming2VideoProducer.h b/Engine/Plugins/Media/PixelStreaming2/Source/PixelStreaming2Core/Public/IPixelStreaming2VideoProducer.h
--- a/Engine/Plugins/Media/PixelStreaming2/Source/PixelStreaming2Core/Public/IPixelStreaming2VideoProducer.h
+++ b/Engine/Plugins/Media/PixelStreaming2/Source/PixelStreaming2Core/Public/IPixelStreaming2VideoProducer.h
@@ -23,7 +23,7 @@ class UPixelStreaming2VideoProducer : public UInterface
  * Note: Your frame is likely a `FPixelCaptureInputFrameRHI` if the frame is a GPU texture or
  * a `FPixelCaptureInputFrameI420` or `FPixelCaptureInputFrameNV12` if your frame is from the CPU.
  */
-class PIXELSTREAMING2CORE_API IPixelStreaming2VideoProducer
+class IPixelStreaming2VideoProducer
 {
 	GENERATED_BODY()
 
@@ -39,4 +39,4 @@ public:
 	 * @return A string containing the display name.
 	 */
 	virtual FString ToString() = 0;
-};
\ No newline at end of file
+};
diff --git a/Engine/Plugins/Media/PixelStreaming2/Source/PixelStreaming2Core/Public/IPixelStreaming2VideoSink.h b/Engine/Plugins/Media/PixelStreaming2/Source/PixelStreaming2Core/Public/IPixelStreaming2VideoSink.h
--- a/Engine/Plugins/Media/PixelStreaming2/Source/PixelStreaming2Core/Public/IPixelStreaming2VideoSink.h
+++ b/Engine/Plugins/Media/PixelStreaming2/Source/PixelStreaming2Core/Public/IPixelStreaming2VideoSink.h
@@ -19,7 +19,7 @@ class IPixelStreaming2VideoConsumer;
 /**
  * Interface for a sink that collects video coming in from the browser and passes into into UE.
  */
-class PIXELSTREAMING2CORE_API IPixelStreaming2VideoSink
+class IPixelStreaming2VideoSink
 {
 	GENERATED_BODY()
 
diff --git a/Engine/Plugins/Media/PixelStreaming2/Source/PixelStreaming2Editor/Internal/PixelStreaming2Style.h b/Engine/Plugins/Media/PixelStreaming2/Source/PixelStreaming2Editor/Internal/PixelStreaming2Style.h
--- a/Engine/Plugins/Media/PixelStreaming2/Source/PixelStreaming2Editor/Internal/PixelStreaming2Style.h
+++ b/Engine/Plugins/Media/PixelStreaming2/Source/PixelStreaming2Editor/Internal/PixelStreaming2Style.h
@@ -3,23 +3,27 @@
 #pragma once
 #include "Styling/SlateStyle.h"
 
+#define UE_API PIXELSTREAMING2EDITOR_API
+
 namespace UE::EditorPixelStreaming2
 {
-    class PIXELSTREAMING2EDITOR_API FPixelStreaming2Style
+    class FPixelStreaming2Style
     {
     public:
-        static void Initialize();
-        static void Shutdown();
+        static UE_API void Initialize();
+        static UE_API void Shutdown();
         /** reloads textures used by slate renderer */
-        static void ReloadTextures();
+        static UE_API void ReloadTextures();
         /** @return The Slate style set for the Shooter game */
-        static FSlateStyleSet& Get();
-        static FName GetStyleSetName();
+        static UE_API FSlateStyleSet& Get();
+        static UE_API FName GetStyleSetName();
    
     private:
-        static TSharedRef< class FSlateStyleSet > Create();
+        static UE_API TSharedRef< class FSlateStyleSet > Create();
    
     private:
-        static TSharedPtr< class FSlateStyleSet > StyleInstance;
+        static UE_API TSharedPtr< class FSlateStyleSet > StyleInstance;
     };
-}
\ No newline at end of file
+}
+
+#undef UE_API
diff --git a/Engine/Plugins/Media/PixelStreaming2/Source/PixelStreaming2Editor/Internal/VideoProducerBackBufferComposited.h b/Engine/Plugins/Media/PixelStreaming2/Source/PixelStreaming2Editor/Internal/VideoProducerBackBufferComposited.h
--- a/Engine/Plugins/Media/PixelStreaming2/Source/PixelStreaming2Editor/Internal/VideoProducerBackBufferComposited.h
+++ b/Engine/Plugins/Media/PixelStreaming2/Source/PixelStreaming2Editor/Internal/VideoProducerBackBufferComposited.h
@@ -8,6 +8,8 @@
 #include "RendererInterface.h"
 #include "Delegates/IDelegateInstance.h"
 
+#define UE_API PIXELSTREAMING2EDITOR_API
+
 namespace UE::EditorPixelStreaming2
 {
 	using namespace UE::PixelStreaming2;
@@ -15,13 +17,13 @@ namespace UE::EditorPixelStreaming2
 	/**
 	 * Use this if you want to send the full UE editor as video input.
 	 */
-	class PIXELSTREAMING2EDITOR_API FVideoProducerBackBufferComposited : public FVideoProducer
+	class FVideoProducerBackBufferComposited : public FVideoProducer
 	{
 	public:
-		static TSharedPtr<FVideoProducerBackBufferComposited> Create();
-		virtual ~FVideoProducerBackBufferComposited();
+		static UE_API TSharedPtr<FVideoProducerBackBufferComposited> Create();
+		UE_API virtual ~FVideoProducerBackBufferComposited();
 
-		virtual FString ToString() override;
+		UE_API virtual FString ToString() override;
 
 		virtual bool IsFrameAlreadyCopied() override { return true; }
 
@@ -57,11 +59,11 @@ namespace UE::EditorPixelStreaming2
 		};
 
 	private:
-		FVideoProducerBackBufferComposited();
-		void CompositeWindows(TSharedPtr<FVideoProducerUserData> UserData);
+		UE_API FVideoProducerBackBufferComposited();
+		UE_API void CompositeWindows(TSharedPtr<FVideoProducerUserData> UserData);
 
-		void OnBackBufferReady(SWindow& SlateWindow, const FTextureRHIRef& FrameBuffer);
-		void OnPreTick(float DeltaTime);
+		UE_API void OnBackBufferReady(SWindow& SlateWindow, const FTextureRHIRef& FrameBuffer);
+		UE_API void OnPreTick(float DeltaTime);
 
 		FDelegateHandle OnBackBufferReadyToPresentHandle;
 		FDelegateHandle OnPreTickHandle;
@@ -89,4 +91,6 @@ namespace UE::EditorPixelStreaming2
 		}
 	};
 
-} // namespace UE::EditorPixelStreaming2
\ No newline at end of file
+} // namespace UE::EditorPixelStreaming2
+
+#undef UE_API
diff --git a/Engine/Plugins/Media/PixelStreaming2/Source/PixelStreaming2Editor/Internal/VideoProducerLevelEditor.h b/Engine/Plugins/Media/PixelStreaming2/Source/PixelStreaming2Editor/Internal/VideoProducerLevelEditor.h
--- a/Engine/Plugins/Media/PixelStreaming2/Source/PixelStreaming2Editor/Internal/VideoProducerLevelEditor.h
+++ b/Engine/Plugins/Media/PixelStreaming2/Source/PixelStreaming2Editor/Internal/VideoProducerLevelEditor.h
@@ -6,6 +6,8 @@
 #include "VideoProducerViewportBase.h"
 #include "Widgets/SWindow.h"
 
+#define UE_API PIXELSTREAMING2EDITOR_API
+
 class FViewport;
 
 namespace UE::EditorPixelStreaming2
@@ -15,21 +17,23 @@ namespace UE::EditorPixelStreaming2
 	/**
 	 * Use this if you want to send the UE primary scene viewport as video input - will only work in editor.
 	 */
-	class PIXELSTREAMING2EDITOR_API FVideoProducerLevelEditor : public FVideoProducerViewportBase
+	class FVideoProducerLevelEditor : public FVideoProducerViewportBase
 	{
 	public:
-		static TSharedPtr<FVideoProducerLevelEditor> Create();
+		static UE_API TSharedPtr<FVideoProducerLevelEditor> Create();
 		virtual ~FVideoProducerLevelEditor() = default;
 
-		virtual FString ToString() override;
+		UE_API virtual FString ToString() override;
 
 	protected:
-		virtual bool ShouldCaptureViewport() override;
+		UE_API virtual bool ShouldCaptureViewport() override;
 
 	private:
 		FVideoProducerLevelEditor() = default;
 
-		void OnPreTick(float DeltaTime);
+		UE_API void OnPreTick(float DeltaTime);
 	};
 
 } // namespace UE::EditorPixelStreaming2
+
+#undef UE_API
diff --git a/Engine/Plugins/Media/PixelStreaming2/Source/PixelStreaming2Editor/Private/Factories.cpp b/Engine/Plugins/Media/PixelStreaming2/Source/PixelStreaming2Editor/Private/Factories.cpp
--- a/Engine/Plugins/Media/PixelStreaming2/Source/PixelStreaming2Editor/Private/Factories.cpp
+++ b/Engine/Plugins/Media/PixelStreaming2/Source/PixelStreaming2Editor/Private/Factories.cpp
@@ -8,6 +8,8 @@
 #include "Blueprints/PixelStreaming2VideoProducer.h"
 #include "IAssetTools.h"
 
+#include UE_INLINE_GENERATED_CPP_BY_NAME(Factories)
+
 #define LOCTEXT_NAMESPACE "PixelStreaming2"
 
 /**
diff --git a/Engine/Plugins/Media/PixelStreaming2/Source/PixelStreaming2Editor/Public/IPixelStreaming2EditorModule.h b/Engine/Plugins/Media/PixelStreaming2/Source/PixelStreaming2Editor/Public/IPixelStreaming2EditorModule.h
--- a/Engine/Plugins/Media/PixelStreaming2/Source/PixelStreaming2Editor/Public/IPixelStreaming2EditorModule.h
+++ b/Engine/Plugins/Media/PixelStreaming2/Source/PixelStreaming2Editor/Public/IPixelStreaming2EditorModule.h
@@ -12,7 +12,7 @@
  * Public interface that manages Pixel Streaming specific functionality within the Unreal Editor.
  * Provides various features for controlling Pixel Streaming, such as starting and stopping the streamer.
  */
-class PIXELSTREAMING2EDITOR_API IPixelStreaming2EditorModule : public IModuleInterface
+class IPixelStreaming2EditorModule : public IModuleInterface
 {
 public:
 	/**
diff --git a/Engine/Plugins/Media/PixelStreaming2/Source/PixelStreaming2HMD/Private/PixelStreaming2HMD.h b/Engine/Plugins/Media/PixelStreaming2/Source/PixelStreaming2HMD/Private/PixelStreaming2HMD.h
--- a/Engine/Plugins/Media/PixelStreaming2/Source/PixelStreaming2HMD/Private/PixelStreaming2HMD.h
+++ b/Engine/Plugins/Media/PixelStreaming2/Source/PixelStreaming2HMD/Private/PixelStreaming2HMD.h
@@ -6,6 +6,8 @@
 #include "IPixelStreaming2HMD.h"
 #include "SceneViewExtension.h"
 
+#define UE_API PIXELSTREAMING2HMD_API
+
 class APlayerController;
 class FSceneView;
 class FSceneViewFamily;
@@ -14,7 +16,7 @@ class UCanvas;
 /**
  * Pixel Streamed Head Mounted Display
  */
-class PIXELSTREAMING2HMD_API FPixelStreaming2HMD : public IPixelStreaming2HMD, public FHeadMountedDisplayBase, public FHMDSceneViewExtension
+class FPixelStreaming2HMD : public IPixelStreaming2HMD, public FHeadMountedDisplayBase, public FHMDSceneViewExtension
 {
 public:
 	// Begin IXRTrackingSystem
@@ -25,20 +27,20 @@ public:
 	}
 
 	int32			GetXRSystemFlags() const { return EXRSystemFlags::IsHeadMounted; }
-	virtual bool	GetRelativeEyePose(int32 DeviceId, int32 ViewIndex, FQuat& OutOrientation, FVector& OutPosition) override;
+	UE_API virtual bool	GetRelativeEyePose(int32 DeviceId, int32 ViewIndex, FQuat& OutOrientation, FVector& OutPosition) override;
 	virtual void	SetBasePosition(const FVector& InBasePosition) override { BasePosition = InBasePosition; };
 	virtual FVector GetBasePosition() const override { return BasePosition; }
 
-	virtual bool EnumerateTrackedDevices(TArray<int32>& OutDevices, EXRTrackedDeviceType Type = EXRTrackedDeviceType::Any) override;
+	UE_API virtual bool EnumerateTrackedDevices(TArray<int32>& OutDevices, EXRTrackedDeviceType Type = EXRTrackedDeviceType::Any) override;
 
-	virtual void  SetInterpupillaryDistance(float NewInterpupillaryDistance) override;
-	virtual float GetInterpupillaryDistance() const override;
+	UE_API virtual void  SetInterpupillaryDistance(float NewInterpupillaryDistance) override;
+	UE_API virtual float GetInterpupillaryDistance() const override;
 
-	virtual void ResetOrientationAndPosition(float yaw = 0.f) override;
+	UE_API virtual void ResetOrientationAndPosition(float yaw = 0.f) override;
 	virtual void ResetOrientation(float Yaw = 0.f) override {}
 	virtual void ResetPosition() override {}
 
-	virtual bool	 GetCurrentPose(int32 DeviceId, FQuat& CurrentOrientation, FVector& CurrentPosition) override;
+	UE_API virtual bool	 GetCurrentPose(int32 DeviceId, FQuat& CurrentOrientation, FVector& CurrentPosition) override;
 	virtual void	 SetBaseRotation(const FRotator& BaseRot) override {};
 	virtual FRotator GetBaseRotation() const override { return FRotator::ZeroRotator; }
 
@@ -55,37 +57,37 @@ public:
 
 protected:
 	// Begin FXRTrackingSystemBase
-	virtual float GetWorldToMetersScale() const override;
+	UE_API virtual float GetWorldToMetersScale() const override;
 	// End FXRTrackingSystemBase
 
 public:
 	// Begin IHeadMountedDisplay
 	virtual bool IsHMDConnected() override { return true; }
-	virtual bool IsHMDEnabled() const override;
-	virtual void EnableHMD(bool allow = true) override;
-	virtual bool GetHMDMonitorInfo(MonitorInfo&) override;
-	virtual void GetFieldOfView(float& OutHFOVInDegrees, float& OutVFOVInDegrees) const override;
-	virtual bool IsChromaAbCorrectionEnabled() const override;
+	UE_API virtual bool IsHMDEnabled() const override;
+	UE_API virtual void EnableHMD(bool allow = true) override;
+	UE_API virtual bool GetHMDMonitorInfo(MonitorInfo&) override;
+	UE_API virtual void GetFieldOfView(float& OutHFOVInDegrees, float& OutVFOVInDegrees) const override;
+	UE_API virtual bool IsChromaAbCorrectionEnabled() const override;
 	virtual bool GetHMDDistortionEnabled(EShadingPath ShadingPath) const override { return false; }
-	virtual void DrawDistortionMesh_RenderThread(struct FHeadMountedDisplayPassContext& Context, const FIntPoint& TextureSize) override;
+	UE_API virtual void DrawDistortionMesh_RenderThread(struct FHeadMountedDisplayPassContext& Context, const FIntPoint& TextureSize) override;
 	// End IHeadMountedDisplay
 
 	// Begin IStereoRendering
-	virtual bool	IsStereoEnabled() const override;
-	virtual bool	EnableStereo(bool stereo = true) override;
-	virtual void	AdjustViewRect(int32 ViewIndex, int32& X, int32& Y, uint32& SizeX, uint32& SizeY) const override;
-	virtual void	CalculateStereoViewOffset(const int32 ViewIndex, FRotator& ViewRotation, const float InWorldToMeters, FVector& ViewLocation) override;
-	virtual FMatrix GetStereoProjectionMatrix(const int32 ViewIndex) const override;
-	virtual void	GetEyeRenderParams_RenderThread(const struct FHeadMountedDisplayPassContext& Context, FVector2D& EyeToSrcUVScaleValue, FVector2D& EyeToSrcUVOffsetValue) const override;
+	UE_API virtual bool	IsStereoEnabled() const override;
+	UE_API virtual bool	EnableStereo(bool stereo = true) override;
+	UE_API virtual void	AdjustViewRect(int32 ViewIndex, int32& X, int32& Y, uint32& SizeX, uint32& SizeY) const override;
+	UE_API virtual void	CalculateStereoViewOffset(const int32 ViewIndex, FRotator& ViewRotation, const float InWorldToMeters, FVector& ViewLocation) override;
+	UE_API virtual FMatrix GetStereoProjectionMatrix(const int32 ViewIndex) const override;
+	UE_API virtual void	GetEyeRenderParams_RenderThread(const struct FHeadMountedDisplayPassContext& Context, FVector2D& EyeToSrcUVScaleValue, FVector2D& EyeToSrcUVOffsetValue) const override;
 	// End IStereoRendering
 
 	// Begin ISceneViewExtension
-	virtual void SetupViewFamily(FSceneViewFamily& InViewFamily) override;
+	UE_API virtual void SetupViewFamily(FSceneViewFamily& InViewFamily) override;
 	// End ISceneViewExtension
 
-	FPixelStreaming2HMD(const FAutoRegister&);
+	UE_API FPixelStreaming2HMD(const FAutoRegister&);
 
-	virtual ~FPixelStreaming2HMD();
+	UE_API virtual ~FPixelStreaming2HMD();
 
 	/**
 	 * @return	True if the HMD was initialized OK
@@ -93,7 +95,7 @@ public:
 	bool IsInitialized() const { return true; }
 
 	virtual void SetTransform(FTransform Transform) override { CurHmdTransform = Transform; }
-	virtual void SetEyeViews(FTransform Left, FMatrix LeftProj, FTransform Right, FMatrix RightProj, FTransform HMD) override;
+	UE_API virtual void SetEyeViews(FTransform Left, FMatrix LeftProj, FTransform Right, FMatrix RightProj, FTransform HMD) override;
 
 private:
 	FVector	   BasePosition = FVector::ZeroVector;
@@ -116,4 +118,6 @@ private:
 	bool	   bStereoEnabled;
 	bool	   bReceivedTransforms = false;
 	double	   LastResChangeSeconds = 0.0f;
-};
\ No newline at end of file
+};
+
+#undef UE_API
diff --git a/Engine/Plugins/Media/PixelStreaming2/Source/PixelStreaming2HMD/Public/IPixelStreaming2HMD.h b/Engine/Plugins/Media/PixelStreaming2/Source/PixelStreaming2HMD/Public/IPixelStreaming2HMD.h
--- a/Engine/Plugins/Media/PixelStreaming2/Source/PixelStreaming2HMD/Public/IPixelStreaming2HMD.h
+++ b/Engine/Plugins/Media/PixelStreaming2/Source/PixelStreaming2HMD/Public/IPixelStreaming2HMD.h
@@ -8,7 +8,7 @@
  * Interface class that allows for setting the transform data for a Head Mounted Display (HMD).
  * The transform and projection matrix for each eye can be individually set.
  */
-class PIXELSTREAMING2HMD_API IPixelStreaming2HMD
+class IPixelStreaming2HMD
 {
 public:
 	/**
diff --git a/Engine/Plugins/Media/PixelStreaming2/Source/PixelStreaming2HMD/Public/IPixelStreaming2HMDModule.h b/Engine/Plugins/Media/PixelStreaming2/Source/PixelStreaming2HMD/Public/IPixelStreaming2HMDModule.h
--- a/Engine/Plugins/Media/PixelStreaming2/Source/PixelStreaming2HMD/Public/IPixelStreaming2HMDModule.h
+++ b/Engine/Plugins/Media/PixelStreaming2/Source/PixelStreaming2HMD/Public/IPixelStreaming2HMDModule.h
@@ -10,7 +10,7 @@
 /**
  * The public interface provides functionality for interacting with Pixel Streaming and Head Mounted Display (HMD) systems.
  */
-class PIXELSTREAMING2HMD_API IPixelStreaming2HMDModule : public IHeadMountedDisplayModule
+class IPixelStreaming2HMDModule : public IHeadMountedDisplayModule
 {
 public:
 	/**
diff --git a/Engine/Plugins/Media/PixelStreaming2/Source/PixelStreaming2Input/Internal/InputHandler.h b/Engine/Plugins/Media/PixelStreaming2/Source/PixelStreaming2Input/Internal/InputHandler.h
--- a/Engine/Plugins/Media/PixelStreaming2/Source/PixelStreaming2Input/Internal/InputHandler.h
+++ b/Engine/Plugins/Media/PixelStreaming2/Source/PixelStreaming2Input/Internal/InputHandler.h
@@ -7,94 +7,96 @@
 #include "XRMotionControllerBase.h"
 #include "PixelStreaming2HMDEnums.h"
 
+#define UE_API PIXELSTREAMING2INPUT_API
+
 namespace UE::PixelStreaming2Input
 {
-	class PIXELSTREAMING2INPUT_API FPixelStreaming2InputHandler : public IPixelStreaming2InputHandler, public FXRMotionControllerBase
+	class FPixelStreaming2InputHandler : public IPixelStreaming2InputHandler, public FXRMotionControllerBase
 	{
 	public:
-		FPixelStreaming2InputHandler();
-		virtual ~FPixelStreaming2InputHandler();
+		UE_API FPixelStreaming2InputHandler();
+		UE_API virtual ~FPixelStreaming2InputHandler();
 
-		virtual void Tick(float DeltaTime) override;
+		UE_API virtual void Tick(float DeltaTime) override;
 
 		// Poll for controller state and send events if needed
 		virtual void SendControllerEvents() override {};
 
 		// Set which MessageHandler will route input
-		virtual void SetMessageHandler(const TSharedRef<FGenericApplicationMessageHandler>& InTargetHandler) override;
+		UE_API virtual void SetMessageHandler(const TSharedRef<FGenericApplicationMessageHandler>& InTargetHandler) override;
 
 		// Register a custom function to execute when command JSON is received.
-		virtual void SetCommandHandler(const FString& CommandName, const CommandHandlerFn& Handler) override;
+		UE_API virtual void SetCommandHandler(const FString& CommandName, const CommandHandlerFn& Handler) override;
 
-		virtual void SetElevatedCheck(const TFunction<bool(FString)>& CheckFn) override;
-		virtual bool IsElevated(const FString& Id) override;
+		UE_API virtual void SetElevatedCheck(const TFunction<bool(FString)>& CheckFn) override;
+		UE_API virtual bool IsElevated(const FString& Id) override;
 
-		virtual TSharedPtr<IPixelStreaming2DataProtocol> GetToStreamerProtocol() override;
-		virtual TSharedPtr<IPixelStreaming2DataProtocol> GetFromStreamerProtocol() override;
+		UE_API virtual TSharedPtr<IPixelStreaming2DataProtocol> GetToStreamerProtocol() override;
+		UE_API virtual TSharedPtr<IPixelStreaming2DataProtocol> GetFromStreamerProtocol() override;
 
 		// Exec handler to allow console commands to be passed through for debugging
-		virtual bool Exec(UWorld* InWorld, const TCHAR* Cmd, FOutputDevice& Ar) override;
+		UE_API virtual bool Exec(UWorld* InWorld, const TCHAR* Cmd, FOutputDevice& Ar) override;
 
 		/**
 		 * IInputInterface pass through functions
 		 */
-		virtual void				SetChannelValue(int32 ControllerId, FForceFeedbackChannelType ChannelType, float Value) override;
-		virtual void				SetChannelValues(int32 ControllerId, const FForceFeedbackValues& values) override;
-		virtual void				OnMessage(FString SourceId, TArray<uint8> Buffer) override;
-		virtual void				SetTargetWindow(TWeakPtr<SWindow> InWindow) override;
-		virtual TWeakPtr<SWindow>	GetTargetWindow() override;
-		virtual void				SetTargetViewport(TWeakPtr<SViewport> InViewport) override;
-		virtual TWeakPtr<SViewport> GetTargetViewport() override;
-		virtual void				SetTargetScreenRect(TWeakPtr<FIntRect> InScreenRect) override;
-		virtual TWeakPtr<FIntRect>	GetTargetScreenRect() override;
+		UE_API virtual void				SetChannelValue(int32 ControllerId, FForceFeedbackChannelType ChannelType, float Value) override;
+		UE_API virtual void				SetChannelValues(int32 ControllerId, const FForceFeedbackValues& values) override;
+		UE_API virtual void				OnMessage(FString SourceId, TArray<uint8> Buffer) override;
+		UE_API virtual void				SetTargetWindow(TWeakPtr<SWindow> InWindow) override;
+		UE_API virtual TWeakPtr<SWindow>	GetTargetWindow() override;
+		UE_API virtual void				SetTargetViewport(TWeakPtr<SViewport> InViewport) override;
+		UE_API virtual TWeakPtr<SViewport> GetTargetViewport() override;
+		UE_API virtual void				SetTargetScreenRect(TWeakPtr<FIntRect> InScreenRect) override;
+		UE_API virtual TWeakPtr<FIntRect>	GetTargetScreenRect() override;
 		virtual bool				IsFakingTouchEvents() const override { return bFakingTouchEvents; }
-		virtual void				RegisterMessageHandler(const FString& MessageType, const MessageHandlerFn& Handler) override;
-		virtual MessageHandlerFn	FindMessageHandler(const FString& MessageType) override;
+		UE_API virtual void				RegisterMessageHandler(const FString& MessageType, const MessageHandlerFn& Handler) override;
+		UE_API virtual MessageHandlerFn	FindMessageHandler(const FString& MessageType) override;
 		virtual void				SetInputType(EPixelStreaming2InputType InInputType) override { InputType = InInputType; };
 		// IMotionController Interface
-		virtual FName			GetMotionControllerDeviceTypeName() const override;
-		virtual bool			GetControllerOrientationAndPosition(const int32 ControllerIndex, const FName MotionSource, FRotator& OutOrientation, FVector& OutPosition, float WorldToMetersScale) const override;
-		virtual ETrackingStatus GetControllerTrackingStatus(const int32 ControllerIndex, const FName MotionSource) const override;
-		virtual void			EnumerateSources(TArray<FMotionControllerSource>& SourcesOut) const override;
+		UE_API virtual FName			GetMotionControllerDeviceTypeName() const override;
+		UE_API virtual bool			GetControllerOrientationAndPosition(const int32 ControllerIndex, const FName MotionSource, FRotator& OutOrientation, FVector& OutPosition, float WorldToMetersScale) const override;
+		UE_API virtual ETrackingStatus GetControllerTrackingStatus(const int32 ControllerIndex, const FName MotionSource) const override;
+		UE_API virtual void			EnumerateSources(TArray<FMotionControllerSource>& SourcesOut) const override;
 		// End IMotionController Interface
 
-		virtual bool OnKeyChar(TCHAR Character) override;
-		virtual bool OnKeyDown(FKey Key, bool bIsRepeat) override;
-		virtual bool OnKeyUp(FKey Key) override;
-
-		virtual bool OnMouseEnter() override;
-		virtual bool OnMouseLeave() override;
-		virtual bool OnMouseDown(EMouseButtons::Type Button, FIntPoint ScreenPosition) override;
-		virtual bool OnMouseUp(EMouseButtons::Type Button) override;
-		virtual bool OnMouseMove(FIntPoint ScreenLocation, FIntPoint Delta) override;
-		virtual bool OnMouseWheel(FIntPoint ScreenLocation, float MouseWheelDelta) override;
-		virtual bool OnMouseDoubleClick(EMouseButtons::Type Button, FIntPoint ScreenPosition) override;
-
-		virtual bool OnTouchStarted(FIntPoint TouchLocation, int32 TouchIndex, float Force) override;
-		virtual bool OnTouchMoved(FIntPoint TouchLocation, int32 TouchIndex, float Force) override;
-		virtual bool OnTouchEnded(FIntPoint TouchLocation, int32 TouchIndex) override;
-
-		virtual uint8 OnControllerConnected() override;
-		virtual bool  OnControllerAnalog(uint8 ControllerIndex, FKey Key, double AxisValue) override;
-		virtual bool  OnControllerButtonPressed(uint8 ControllerIndex, FKey Key, bool bIsRepeat) override;
-		virtual bool  OnControllerButtonReleased(uint8 ControllerIndex, FKey Key) override;
-		virtual bool  OnControllerDisconnected(uint8 ControllerIndex) override;
-
-		virtual bool OnXREyeViews(FTransform LeftEyeTransform, FMatrix LeftEyeProjectionMatrix, FTransform RightEyeTransform, FMatrix RightEyeProjectionMatrix, FTransform HMDTransform) override;
-		virtual bool OnXRHMDTransform(FTransform HMDTransform) override;
-		virtual bool OnXRControllerTransform(FTransform ControllerTransform, EControllerHand Handedness) override;
-		virtual bool OnXRButtonTouched(EControllerHand Handedness, FKey Key, bool bIsRepeat) override;
-		virtual bool OnXRButtonTouchReleased(EControllerHand Handedness, FKey Key) override;
-		virtual bool OnXRButtonPressed(EControllerHand Handedness, FKey Key, bool bIsRepeat) override;
-		virtual bool OnXRButtonReleased(EControllerHand Handedness, FKey Key) override;
-		virtual bool OnXRAnalog(EControllerHand Handedness, FKey Key, double AnalogValue) override;
-		virtual bool OnXRSystem(EPixelStreaming2XRSystem System) override;
+		UE_API virtual bool OnKeyChar(TCHAR Character) override;
+		UE_API virtual bool OnKeyDown(FKey Key, bool bIsRepeat) override;
+		UE_API virtual bool OnKeyUp(FKey Key) override;
+
+		UE_API virtual bool OnMouseEnter() override;
+		UE_API virtual bool OnMouseLeave() override;
+		UE_API virtual bool OnMouseDown(EMouseButtons::Type Button, FIntPoint ScreenPosition) override;
+		UE_API virtual bool OnMouseUp(EMouseButtons::Type Button) override;
+		UE_API virtual bool OnMouseMove(FIntPoint ScreenLocation, FIntPoint Delta) override;
+		UE_API virtual bool OnMouseWheel(FIntPoint ScreenLocation, float MouseWheelDelta) override;
+		UE_API virtual bool OnMouseDoubleClick(EMouseButtons::Type Button, FIntPoint ScreenPosition) override;
+
+		UE_API virtual bool OnTouchStarted(FIntPoint TouchLocation, int32 TouchIndex, float Force) override;
+		UE_API virtual bool OnTouchMoved(FIntPoint TouchLocation, int32 TouchIndex, float Force) override;
+		UE_API virtual bool OnTouchEnded(FIntPoint TouchLocation, int32 TouchIndex) override;
+
+		UE_API virtual uint8 OnControllerConnected() override;
+		UE_API virtual bool  OnControllerAnalog(uint8 ControllerIndex, FKey Key, double AxisValue) override;
+		UE_API virtual bool  OnControllerButtonPressed(uint8 ControllerIndex, FKey Key, bool bIsRepeat) override;
+		UE_API virtual bool  OnControllerButtonReleased(uint8 ControllerIndex, FKey Key) override;
+		UE_API virtual bool  OnControllerDisconnected(uint8 ControllerIndex) override;
+
+		UE_API virtual bool OnXREyeViews(FTransform LeftEyeTransform, FMatrix LeftEyeProjectionMatrix, FTransform RightEyeTransform, FMatrix RightEyeProjectionMatrix, FTransform HMDTransform) override;
+		UE_API virtual bool OnXRHMDTransform(FTransform HMDTransform) override;
+		UE_API virtual bool OnXRControllerTransform(FTransform ControllerTransform, EControllerHand Handedness) override;
+		UE_API virtual bool OnXRButtonTouched(EControllerHand Handedness, FKey Key, bool bIsRepeat) override;
+		UE_API virtual bool OnXRButtonTouchReleased(EControllerHand Handedness, FKey Key) override;
+		UE_API virtual bool OnXRButtonPressed(EControllerHand Handedness, FKey Key, bool bIsRepeat) override;
+		UE_API virtual bool OnXRButtonReleased(EControllerHand Handedness, FKey Key) override;
+		UE_API virtual bool OnXRAnalog(EControllerHand Handedness, FKey Key, double AnalogValue) override;
+		UE_API virtual bool OnXRSystem(EPixelStreaming2XRSystem System) override;
 
 	protected:
-		FVector2D	ConvertToNormalizedScreenLocation(FVector2D Point);
-		FIntPoint	ConvertFromNormalizedScreenLocation(const FVector2D& ScreenLocation, bool bIncludeOffset = true);
-		FWidgetPath FindRoutingMessageWidget(const FVector2D& Location) const;
-		FKey		TranslateMouseButtonToKey(const EMouseButtons::Type Button);
+		UE_API FVector2D	ConvertToNormalizedScreenLocation(FVector2D Point);
+		UE_API FIntPoint	ConvertFromNormalizedScreenLocation(const FVector2D& ScreenLocation, bool bIncludeOffset = true);
+		UE_API FWidgetPath FindRoutingMessageWidget(const FVector2D& Location) const;
+		UE_API FKey		TranslateMouseButtonToKey(const EMouseButtons::Type Button);
 
 		struct FCachedTouchEvent
 		{
@@ -134,25 +136,25 @@ namespace UE::PixelStreaming2Input
 		TMap<FInputDeviceId, TMap<FKey, FAnalogValue>> AnalogEventsReceivedThisTick;
 
 		/** Forwards the latest analog input received for each key this tick. */
-		void ProcessLatestAnalogInputFromThisTick();
+		UE_API void ProcessLatestAnalogInputFromThisTick();
 
 		// Track which touch events we processed this frame so we can avoid re-processing them
 		TSet<int32> TouchIndicesProcessedThisFrame;
 
 		// Sends Touch Moved events for any touch index which is currently down but wasn't already updated this frame
-		void BroadcastActiveTouchMoveEvents();
+		UE_API void BroadcastActiveTouchMoveEvents();
 
-		void FindFocusedWidget();
+		UE_API void FindFocusedWidget();
 
 		/**
 		 * Create an artificial mouse 'movement' to allow widgets to focus under a
 		 * static mouse cursor. This will not actually change the cursor position.
 		 */
-		void SynthesizeMouseMove() const;
+		UE_API void SynthesizeMouseMove() const;
 
 		TArray<FKey> FilteredKeys;
-		void		 OnFilteredKeysChanged(IConsoleVariable* Var);
-		bool		 FilterKey(const FKey& Key);
+		UE_API void		 OnFilteredKeysChanged(IConsoleVariable* Var);
+		UE_API bool		 FilterKey(const FKey& Key);
 
 		struct FMessage
 		{
@@ -241,4 +243,6 @@ namespace UE::PixelStreaming2Input
 	private:
 		FDelegateHandle OnInputKeyFilterChangedHandle;
 	};
-} // namespace UE::PixelStreaming2Input
\ No newline at end of file
+} // namespace UE::PixelStreaming2Input
+
+#undef UE_API
diff --git a/Engine/Plugins/Media/PixelStreaming2/Source/PixelStreaming2Input/Internal/InputProtocolMap.h b/Engine/Plugins/Media/PixelStreaming2/Source/PixelStreaming2Input/Internal/InputProtocolMap.h
--- a/Engine/Plugins/Media/PixelStreaming2/Source/PixelStreaming2Input/Internal/InputProtocolMap.h
+++ b/Engine/Plugins/Media/PixelStreaming2/Source/PixelStreaming2Input/Internal/InputProtocolMap.h
@@ -6,6 +6,8 @@
 #include "IPixelStreaming2DataProtocol.h"
 #include "PixelStreaming2InputEnums.h"
 
+#define UE_API PIXELSTREAMING2INPUT_API
+
 namespace UE::PixelStreaming2Input
 {
 
@@ -13,32 +15,32 @@ namespace UE::PixelStreaming2Input
 	 * @brief An map type that broadcasts the OnProtocolUpdated whenever
 	 * it's inner map is updated
 	 */
-	class PIXELSTREAMING2INPUT_API FInputProtocolMap : public IPixelStreaming2DataProtocol
+	class FInputProtocolMap : public IPixelStreaming2DataProtocol
 	{
 	public:
 		FInputProtocolMap(EPixelStreaming2MessageDirection InDirection)
 			: Direction(InDirection) {}
 
 		// Begin IPixelStreaming2DataProtocol interface
-		TSharedPtr<IPixelStreaming2InputMessage> Add(FString Key) override;
-		TSharedPtr<IPixelStreaming2InputMessage> Add(FString Key, TArray<EPixelStreaming2MessageTypes> InStructure) override;
-		TSharedPtr<IPixelStreaming2InputMessage> Find(FString Key) override;
+		UE_API TSharedPtr<IPixelStreaming2InputMessage> Add(FString Key) override;
+		UE_API TSharedPtr<IPixelStreaming2InputMessage> Add(FString Key, TArray<EPixelStreaming2MessageTypes> InStructure) override;
+		UE_API TSharedPtr<IPixelStreaming2InputMessage> Find(FString Key) override;
 		FOnProtocolUpdated&						 OnProtocolUpdated() override { return OnProtocolUpdatedDelegate; };
-		TSharedPtr<FJsonObject>					 ToJson() override;
+		UE_API TSharedPtr<FJsonObject>					 ToJson() override;
 		// End IPixelStreaming2DataProtocol interface
 
-		bool										   AddInternal(FString Key, uint8 Id);
-		bool										   AddInternal(FString Key, uint8 Id, TArray<EPixelStreaming2MessageTypes> InStructure);
-		int											   Remove(FString Key);
-		const TSharedPtr<IPixelStreaming2InputMessage> Find(FString Key) const;
+		UE_API bool										   AddInternal(FString Key, uint8 Id);
+		UE_API bool										   AddInternal(FString Key, uint8 Id, TArray<EPixelStreaming2MessageTypes> InStructure);
+		UE_API int											   Remove(FString Key);
+		UE_API const TSharedPtr<IPixelStreaming2InputMessage> Find(FString Key) const;
 
-		void Clear();
-		bool IsEmpty() const;
+		UE_API void Clear();
+		UE_API bool IsEmpty() const;
 
-		void Apply(const TFunction<void(FString, TSharedPtr<IPixelStreaming2InputMessage>)>& Visitor);
+		UE_API void Apply(const TFunction<void(FString, TSharedPtr<IPixelStreaming2InputMessage>)>& Visitor);
 
 	private:
-		TSharedPtr<IPixelStreaming2InputMessage> AddMessageInternal(FString Key, uint8 Id, TArray<EPixelStreaming2MessageTypes> InStructure);
+		UE_API TSharedPtr<IPixelStreaming2InputMessage> AddMessageInternal(FString Key, uint8 Id, TArray<EPixelStreaming2MessageTypes> InStructure);
 
 	private:
 		TSet<uint8>												Ids;
@@ -49,3 +51,5 @@ namespace UE::PixelStreaming2Input
 	};
 
 } // namespace UE::PixelStreaming2Input
+
+#undef UE_API
diff --git a/Engine/Plugins/Media/PixelStreaming2/Source/PixelStreaming2Input/Public/IPixelStreaming2DataProtocol.h b/Engine/Plugins/Media/PixelStreaming2/Source/PixelStreaming2Input/Public/IPixelStreaming2DataProtocol.h
--- a/Engine/Plugins/Media/PixelStreaming2/Source/PixelStreaming2Input/Public/IPixelStreaming2DataProtocol.h
+++ b/Engine/Plugins/Media/PixelStreaming2/Source/PixelStreaming2Input/Public/IPixelStreaming2DataProtocol.h
@@ -10,7 +10,7 @@
  * Interface for a datachannel protocol.
  * The user is not meant to extend this through polymorphism, but rather request the appropriate protocol from the InputHandler and extend it using Add().
  */
-class PIXELSTREAMING2INPUT_API IPixelStreaming2DataProtocol
+class IPixelStreaming2DataProtocol
 {
 public:
 	IPixelStreaming2DataProtocol() = default;
@@ -50,4 +50,4 @@ public:
 	 * Internally this will force Pixel Streaming to resend the entire data protocol.
 	 */
 	virtual FOnProtocolUpdated& OnProtocolUpdated() = 0;
-};
\ No newline at end of file
+};
diff --git a/Engine/Plugins/Media/PixelStreaming2/Source/PixelStreaming2Input/Public/IPixelStreaming2InputHandler.h b/Engine/Plugins/Media/PixelStreaming2/Source/PixelStreaming2Input/Public/IPixelStreaming2InputHandler.h
--- a/Engine/Plugins/Media/PixelStreaming2/Source/PixelStreaming2Input/Public/IPixelStreaming2InputHandler.h
+++ b/Engine/Plugins/Media/PixelStreaming2/Source/PixelStreaming2Input/Public/IPixelStreaming2InputHandler.h
@@ -13,7 +13,7 @@
  * scaling of input from browser to application, and setting the target window ensure that if windows are tiled (eg editor)
  * that the streamed input only affect the target window.
  */
-class PIXELSTREAMING2INPUT_API IPixelStreaming2InputHandler : public IInputDevice
+class IPixelStreaming2InputHandler : public IInputDevice
 {
 public:
 	/**
@@ -409,4 +409,4 @@ public:
 	 * @return true if the event was handled. false otherwise
 	 */
 	virtual bool OnXRSystem(EPixelStreaming2XRSystem System) = 0;
-};
\ No newline at end of file
+};
diff --git a/Engine/Plugins/Media/PixelStreaming2/Source/PixelStreaming2Input/Public/IPixelStreaming2InputMessage.h b/Engine/Plugins/Media/PixelStreaming2/Source/PixelStreaming2Input/Public/IPixelStreaming2InputMessage.h
--- a/Engine/Plugins/Media/PixelStreaming2/Source/PixelStreaming2Input/Public/IPixelStreaming2InputMessage.h
+++ b/Engine/Plugins/Media/PixelStreaming2/Source/PixelStreaming2Input/Public/IPixelStreaming2InputMessage.h
@@ -9,7 +9,7 @@
 /**
  * A message that is sent over the Pixel Streaming data channel.
  */
-class PIXELSTREAMING2INPUT_API IPixelStreaming2InputMessage
+class IPixelStreaming2InputMessage
 {
 public:
 	/**
@@ -25,4 +25,4 @@ public:
 	 * @return The structure of the message.
 	 */
 	virtual TArray<EPixelStreaming2MessageTypes> GetStructure() const = 0;
-};
\ No newline at end of file
+};
diff --git a/Engine/Plugins/Media/PixelStreaming2/Source/PixelStreaming2Input/Public/IPixelStreaming2InputModule.h b/Engine/Plugins/Media/PixelStreaming2/Source/PixelStreaming2Input/Public/IPixelStreaming2InputModule.h
--- a/Engine/Plugins/Media/PixelStreaming2/Source/PixelStreaming2Input/Public/IPixelStreaming2InputModule.h
+++ b/Engine/Plugins/Media/PixelStreaming2/Source/PixelStreaming2Input/Public/IPixelStreaming2InputModule.h
@@ -8,7 +8,7 @@
 /**
  * Provides mechanism for managing input handlers and devices in Pixel Streaming.
  */
-class PIXELSTREAMING2INPUT_API IPixelStreaming2InputModule : public IInputDeviceModule
+class IPixelStreaming2InputModule : public IInputDeviceModule
 {
 public:
 	/**
diff --git a/Engine/Plugins/Media/PixelStreaming2/Source/PixelStreaming2RTC/Internal/AtomicOnceFlag.h b/Engine/Plugins/Media/PixelStreaming2/Source/PixelStreaming2RTC/Internal/AtomicOnceFlag.h
--- a/Engine/Plugins/Media/PixelStreaming2/Source/PixelStreaming2RTC/Internal/AtomicOnceFlag.h
+++ b/Engine/Plugins/Media/PixelStreaming2/Source/PixelStreaming2RTC/Internal/AtomicOnceFlag.h
@@ -8,7 +8,7 @@
 
 namespace UE::PixelStreaming2
 {
-	class PIXELSTREAMING2RTC_API FAtomicOnceFlag
+	class FAtomicOnceFlag
 	{
 	public:
 		// Sets flag to true and returns if the flag had fired before
diff --git a/Engine/Plugins/Media/PixelStreaming2/Source/PixelStreaming2RTC/Internal/EpicRtcAllocator.h b/Engine/Plugins/Media/PixelStreaming2/Source/PixelStreaming2RTC/Internal/EpicRtcAllocator.h
--- a/Engine/Plugins/Media/PixelStreaming2/Source/PixelStreaming2RTC/Internal/EpicRtcAllocator.h
+++ b/Engine/Plugins/Media/PixelStreaming2/Source/PixelStreaming2RTC/Internal/EpicRtcAllocator.h
@@ -8,7 +8,7 @@
 
 namespace UE::PixelStreaming2
 {
-	class PIXELSTREAMING2RTC_API FEpicRtcAllocator : public EpicRtcMemoryInterface
+	class FEpicRtcAllocator : public EpicRtcMemoryInterface
 	{
 	public:
 		FEpicRtcAllocator() = default;
@@ -29,4 +29,4 @@ namespace UE::PixelStreaming2
 			return FMemory::Free(Pointer);
 		}
 	};
-} // namespace UE::PixelStreaming2
\ No newline at end of file
+} // namespace UE::PixelStreaming2
diff --git a/Engine/Plugins/Media/PixelStreaming2/Source/PixelStreaming2RTC/Internal/EpicRtcAudioCapturer.h b/Engine/Plugins/Media/PixelStreaming2/Source/PixelStreaming2RTC/Internal/EpicRtcAudioCapturer.h
--- a/Engine/Plugins/Media/PixelStreaming2/Source/PixelStreaming2RTC/Internal/EpicRtcAudioCapturer.h
+++ b/Engine/Plugins/Media/PixelStreaming2/Source/PixelStreaming2RTC/Internal/EpicRtcAudioCapturer.h
@@ -4,20 +4,22 @@
 
 #include "AudioCapturer.h"
 
+#define UE_API PIXELSTREAMING2RTC_API
+
 namespace UE::PixelStreaming2
 {
 	/**
 	 * FEpicRtcAudioCapturer overrides the default PushAudio behaviour of the FAudioCapturer in order to
 	 * break up the pushed audio into 10ms chunks
 	 */
-	class PIXELSTREAMING2RTC_API FEpicRtcAudioCapturer : public FAudioCapturer
+	class FEpicRtcAudioCapturer : public FAudioCapturer
 	{
 	public:
-		static TSharedPtr<FEpicRtcAudioCapturer> Create();
+		static UE_API TSharedPtr<FEpicRtcAudioCapturer> Create();
 		virtual ~FEpicRtcAudioCapturer() = default;
 
 		// Override the push audio method as EpicRtc needs the broadcasted audio to be in 10ms chunks
-		virtual void PushAudio(const float* AudioData, int32 NumSamples, int32 NumChannels, int32 SampleRate) override;
+		UE_API virtual void PushAudio(const float* AudioData, int32 NumSamples, int32 NumChannels, int32 SampleRate) override;
 
 	protected:
 		FEpicRtcAudioCapturer() = default;
@@ -25,4 +27,6 @@ namespace UE::PixelStreaming2
 	private:
 		TArray<int16_t> RecordingBuffer;
 	};
-} // namespace UE::PixelStreaming2
\ No newline at end of file
+} // namespace UE::PixelStreaming2
+
+#undef UE_API
diff --git a/Engine/Plugins/Media/PixelStreaming2/Source/PixelStreaming2RTC/Internal/EpicRtcAudioSink.h b/Engine/Plugins/Media/PixelStreaming2/Source/PixelStreaming2RTC/Internal/EpicRtcAudioSink.h
--- a/Engine/Plugins/Media/PixelStreaming2/Source/PixelStreaming2RTC/Internal/EpicRtcAudioSink.h
+++ b/Engine/Plugins/Media/PixelStreaming2/Source/PixelStreaming2RTC/Internal/EpicRtcAudioSink.h
@@ -6,19 +6,23 @@
 #include "AudioSink.h"
 #include "EpicRtcTrack.h"
 
+#define UE_API PIXELSTREAMING2RTC_API
+
 namespace UE::PixelStreaming2
 {
 	// Collects audio coming in from EpicRtc and passes into into UE's audio system.
-	class PIXELSTREAMING2RTC_API FEpicRtcAudioSink : public FAudioSink, public TEpicRtcTrack<EpicRtcAudioTrackInterface>
+	class FEpicRtcAudioSink : public FAudioSink, public TEpicRtcTrack<EpicRtcAudioTrackInterface>
 	{
 	public:
-		static TSharedPtr<FEpicRtcAudioSink> Create(TRefCountPtr<EpicRtcAudioTrackInterface> InTrack);
+		static UE_API TSharedPtr<FEpicRtcAudioSink> Create(TRefCountPtr<EpicRtcAudioTrackInterface> InTrack);
 		virtual ~FEpicRtcAudioSink() = default;
 		
 		inline bool TriggerDisconnectFlag() { return bDisconnectFlag.Trigger(); };
 	private:
-		FEpicRtcAudioSink(TRefCountPtr<EpicRtcAudioTrackInterface> InTrack);
+		UE_API FEpicRtcAudioSink(TRefCountPtr<EpicRtcAudioTrackInterface> InTrack);
 
 		FAtomicOnceFlag bDisconnectFlag;
 	};
-} // namespace UE::PixelStreaming2
\ No newline at end of file
+} // namespace UE::PixelStreaming2
+
+#undef UE_API
diff --git a/Engine/Plugins/Media/PixelStreaming2/Source/PixelStreaming2RTC/Internal/EpicRtcAudioSource.h b/Engine/Plugins/Media/PixelStreaming2/Source/PixelStreaming2RTC/Internal/EpicRtcAudioSource.h
--- a/Engine/Plugins/Media/PixelStreaming2/Source/PixelStreaming2RTC/Internal/EpicRtcAudioSource.h
+++ b/Engine/Plugins/Media/PixelStreaming2/Source/PixelStreaming2RTC/Internal/EpicRtcAudioSource.h
@@ -9,19 +9,23 @@
 
 #include "epic_rtc/core/audio/audio_track.h"
 
+#define UE_API PIXELSTREAMING2RTC_API
+
 namespace UE::PixelStreaming2
 {
-	class PIXELSTREAMING2RTC_API FEpicRtcAudioSource : public FAudioSource, public TEpicRtcTrack<EpicRtcAudioTrackInterface>
+	class FEpicRtcAudioSource : public FAudioSource, public TEpicRtcTrack<EpicRtcAudioTrackInterface>
 	{
 	public:
-		static TSharedPtr<FEpicRtcAudioSource> Create(TRefCountPtr<EpicRtcAudioTrackInterface> InTrack, TSharedPtr<class FEpicRtcAudioCapturer> InCapturer);
+		static UE_API TSharedPtr<FEpicRtcAudioSource> Create(TRefCountPtr<EpicRtcAudioTrackInterface> InTrack, TSharedPtr<class FEpicRtcAudioCapturer> InCapturer);
 
-		virtual void OnAudioBuffer(const int16_t* AudioData, int32 NumSamples, int32 NumChannels, const int32 SampleRate);
+		UE_API virtual void OnAudioBuffer(const int16_t* AudioData, int32 NumSamples, int32 NumChannels, const int32 SampleRate);
 
 		inline bool TriggerDisconnectFlag() { return bDisconnectFlag.Trigger(); };
 	private:
-		FEpicRtcAudioSource(TRefCountPtr<EpicRtcAudioTrackInterface> InTrack);
+		UE_API FEpicRtcAudioSource(TRefCountPtr<EpicRtcAudioTrackInterface> InTrack);
 
 		FAtomicOnceFlag bDisconnectFlag;
 	};
-} // namespace UE::PixelStreaming2
\ No newline at end of file
+} // namespace UE::PixelStreaming2
+
+#undef UE_API
diff --git a/Engine/Plugins/Media/PixelStreaming2/Source/PixelStreaming2RTC/Internal/EpicRtcAudioTrackObserver.h b/Engine/Plugins/Media/PixelStreaming2/Source/PixelStreaming2RTC/Internal/EpicRtcAudioTrackObserver.h
--- a/Engine/Plugins/Media/PixelStreaming2/Source/PixelStreaming2RTC/Internal/EpicRtcAudioTrackObserver.h
+++ b/Engine/Plugins/Media/PixelStreaming2/Source/PixelStreaming2RTC/Internal/EpicRtcAudioTrackObserver.h
@@ -9,13 +9,15 @@
 
 #include "EpicRtcAudioTrackObserver.generated.h"
 
+#define UE_API PIXELSTREAMING2RTC_API
+
 UINTERFACE(MinimalAPI)
 class UPixelStreaming2AudioTrackObserver : public UInterface
 {
 	GENERATED_BODY()
 };
 
-class PIXELSTREAMING2RTC_API IPixelStreaming2AudioTrackObserver
+class IPixelStreaming2AudioTrackObserver
 {
 	GENERATED_BODY()
 
@@ -28,18 +30,18 @@ public:
 
 namespace UE::PixelStreaming2
 {
-	class PIXELSTREAMING2RTC_API FEpicRtcAudioTrackObserver : public EpicRtcAudioTrackObserverInterface
+	class FEpicRtcAudioTrackObserver : public EpicRtcAudioTrackObserverInterface
 	{
 	public:
-		FEpicRtcAudioTrackObserver(TObserverVariant<IPixelStreaming2AudioTrackObserver> UserObserver);
+		UE_API FEpicRtcAudioTrackObserver(TObserverVariant<IPixelStreaming2AudioTrackObserver> UserObserver);
 		virtual ~FEpicRtcAudioTrackObserver() = default;
 
 	private:
 		// Begin EpicRtcAudioTrackObserverInterface
-		virtual void OnAudioTrackMuted(EpicRtcAudioTrackInterface* AudioTrack, EpicRtcBool bIsMuted) override;
-		virtual void OnAudioTrackFrame(EpicRtcAudioTrackInterface* AudioTrack, const EpicRtcAudioFrame& Frame) override;
-		virtual void OnAudioTrackRemoved(EpicRtcAudioTrackInterface* AudioTrack) override;
-		virtual void OnAudioTrackState(EpicRtcAudioTrackInterface* AudioTrack, const EpicRtcTrackState State) override;
+		UE_API virtual void OnAudioTrackMuted(EpicRtcAudioTrackInterface* AudioTrack, EpicRtcBool bIsMuted) override;
+		UE_API virtual void OnAudioTrackFrame(EpicRtcAudioTrackInterface* AudioTrack, const EpicRtcAudioFrame& Frame) override;
+		UE_API virtual void OnAudioTrackRemoved(EpicRtcAudioTrackInterface* AudioTrack) override;
+		UE_API virtual void OnAudioTrackState(EpicRtcAudioTrackInterface* AudioTrack, const EpicRtcTrackState State) override;
 		// End EpicRtcAudioTrackObserverInterface
 
 	public:
@@ -50,4 +52,6 @@ namespace UE::PixelStreaming2
 	private:
 		TObserverVariant<IPixelStreaming2AudioTrackObserver> UserObserver;
 	};
-} // namespace UE::PixelStreaming2
\ No newline at end of file
+} // namespace UE::PixelStreaming2
+
+#undef UE_API
diff --git a/Engine/Plugins/Media/PixelStreaming2/Source/PixelStreaming2RTC/Internal/EpicRtcAudioTrackObserverFactory.h b/Engine/Plugins/Media/PixelStreaming2/Source/PixelStreaming2RTC/Internal/EpicRtcAudioTrackObserverFactory.h
--- a/Engine/Plugins/Media/PixelStreaming2/Source/PixelStreaming2RTC/Internal/EpicRtcAudioTrackObserverFactory.h
+++ b/Engine/Plugins/Media/PixelStreaming2/Source/PixelStreaming2RTC/Internal/EpicRtcAudioTrackObserverFactory.h
@@ -7,19 +7,21 @@
 #include "epic_rtc/core/audio/audio_track_observer.h"
 #include "epic_rtc_helper/memory/ref_count_impl_helper.h"
 
+#define UE_API PIXELSTREAMING2RTC_API
+
 class IPixelStreaming2AudioTrackObserver;
 
 namespace UE::PixelStreaming2
 {		
-	class PIXELSTREAMING2RTC_API FEpicRtcAudioTrackObserverFactory : public EpicRtcAudioTrackObserverFactoryInterface
+	class FEpicRtcAudioTrackObserverFactory : public EpicRtcAudioTrackObserverFactoryInterface
 	{
 	public:
-		FEpicRtcAudioTrackObserverFactory(TObserverVariant<IPixelStreaming2AudioTrackObserver> UserObserver);
+		UE_API FEpicRtcAudioTrackObserverFactory(TObserverVariant<IPixelStreaming2AudioTrackObserver> UserObserver);
 		virtual ~FEpicRtcAudioTrackObserverFactory() = default;
 
 	public:
 		// Begin EpicRtcAudioTrackObserverFactoryInterface
-		virtual EpicRtcErrorCode CreateAudioTrackObserver(const EpicRtcStringView ParticipantId, const EpicRtcStringView AudioTrackId, EpicRtcAudioTrackObserverInterface** OutAudioTrackObserver) override;
+		UE_API virtual EpicRtcErrorCode CreateAudioTrackObserver(const EpicRtcStringView ParticipantId, const EpicRtcStringView AudioTrackId, EpicRtcAudioTrackObserverInterface** OutAudioTrackObserver) override;
 		// End EpicRtcAudioTrackObserverFactoryInterface
 
 	public:
@@ -31,4 +33,6 @@ namespace UE::PixelStreaming2
 		TObserverVariant<IPixelStreaming2AudioTrackObserver> UserObserver;
 	};
 
-} // namespace UE::PixelStreaming2
\ No newline at end of file
+} // namespace UE::PixelStreaming2
+
+#undef UE_API
diff --git a/Engine/Plugins/Media/PixelStreaming2/Source/PixelStreaming2RTC/Internal/EpicRtcConferenceUtils.h b/Engine/Plugins/Media/PixelStreaming2/Source/PixelStreaming2RTC/Internal/EpicRtcConferenceUtils.h
--- a/Engine/Plugins/Media/PixelStreaming2/Source/PixelStreaming2RTC/Internal/EpicRtcConferenceUtils.h
+++ b/Engine/Plugins/Media/PixelStreaming2/Source/PixelStreaming2RTC/Internal/EpicRtcConferenceUtils.h
@@ -11,7 +11,7 @@
 
 namespace UE::PixelStreaming2
 {
-	class PIXELSTREAMING2RTC_API FEpicRtcTickConferenceTask : public FPixelStreamingTickableTask
+	class FEpicRtcTickConferenceTask : public FPixelStreamingTickableTask
 	{
 	public:
 		FEpicRtcTickConferenceTask(TRefCountPtr<EpicRtcConferenceInterface>& EpicRtcConference, const FString& TaskName = TEXT("EpicRtcTickConferenceTask"))
@@ -64,4 +64,4 @@ namespace UE::PixelStreaming2
 		FString									 TaskName;
 	};
 
-} // namespace UE::PixelStreaming2
\ No newline at end of file
+} // namespace UE::PixelStreaming2
