The source files included were modified by the UnrealCodeFixup tool so that they can pass the -mergemodules compilation.

diff --git a/Engine/Plugins/Media/PixelStreaming2/Source/PixelStreaming2RTC/Internal/EpicRtcDataTrack.h b/Engine/Plugins/Media/PixelStreaming2/Source/PixelStreaming2RTC/Internal/EpicRtcDataTrack.h
--- a/Engine/Plugins/Media/PixelStreaming2/Source/PixelStreaming2RTC/Internal/EpicRtcDataTrack.h
+++ b/Engine/Plugins/Media/PixelStreaming2/Source/PixelStreaming2RTC/Internal/EpicRtcDataTrack.h
@@ -6,6 +6,8 @@
 #include "IPixelStreaming2DataProtocol.h"
 #include "Templates/SharedPointer.h"
 
+#define UE_API PIXELSTREAMING2RTC_API
+
 namespace UE::PixelStreaming2
 {
 	template <typename T>
@@ -60,7 +62,7 @@ namespace UE::PixelStreaming2
 		return Value.GetData();
 	}
 
-	struct PIXELSTREAMING2RTC_API FBufferBuilder
+	struct FBufferBuilder
 	{
 		TArray<uint8> Buffer;
 		size_t		  Pos;
@@ -93,10 +95,10 @@ namespace UE::PixelStreaming2
 		}
 	};
 
-	class PIXELSTREAMING2RTC_API FEpicRtcDataTrack : public TEpicRtcTrack<EpicRtcDataTrackInterface>, public TSharedFromThis<FEpicRtcDataTrack>
+	class FEpicRtcDataTrack : public TEpicRtcTrack<EpicRtcDataTrackInterface>, public TSharedFromThis<FEpicRtcDataTrack>
 	{
 	public:
-		static TSharedPtr<FEpicRtcDataTrack> Create(TRefCountPtr<EpicRtcDataTrackInterface> InTrack, TWeakPtr<IPixelStreaming2DataProtocol> InDataProtocol);
+		static UE_API TSharedPtr<FEpicRtcDataTrack> Create(TRefCountPtr<EpicRtcDataTrackInterface> InTrack, TWeakPtr<IPixelStreaming2DataProtocol> InDataProtocol);
 
 		virtual ~FEpicRtcDataTrack() = default;
 
@@ -130,7 +132,7 @@ namespace UE::PixelStreaming2
 		 * @param DataBytes The raw byte buffer to send.
 		 * @returns True of the message was successfully sent.
 		 */
-		bool SendArbitraryData(const FString& MessageType, const TArray64<uint8>& DataBytes) const;
+		UE_API bool SendArbitraryData(const FString& MessageType, const TArray64<uint8>& DataBytes) const;
 
 		/**
 		 * @return The state of the underlying EpicRtc data track.
@@ -140,14 +142,14 @@ namespace UE::PixelStreaming2
 		void SetSendTrack(TRefCountPtr<EpicRtcDataTrackInterface> InSendTrack) { SendTrack = InSendTrack; }
 
 	protected:
-		FEpicRtcDataTrack(TRefCountPtr<EpicRtcDataTrackInterface> InTrack, TWeakPtr<IPixelStreaming2DataProtocol> InDataProtocol);
-		FEpicRtcDataTrack(TSharedPtr<FEpicRtcDataTrack> InTrack, TWeakPtr<IPixelStreaming2DataProtocol> InDataProtocol);
+		UE_API FEpicRtcDataTrack(TRefCountPtr<EpicRtcDataTrackInterface> InTrack, TWeakPtr<IPixelStreaming2DataProtocol> InDataProtocol);
+		UE_API FEpicRtcDataTrack(TSharedPtr<FEpicRtcDataTrack> InTrack, TWeakPtr<IPixelStreaming2DataProtocol> InDataProtocol);
 
 		virtual void PrependData(FBufferBuilder& Builder) const {};
 
-		bool IsActive() const;
+		UE_API bool IsActive() const;
 
-		bool GetMessageId(const FString& MessageType, uint8& OutMessageId) const;
+		UE_API bool GetMessageId(const FString& MessageType, uint8& OutMessageId) const;
 
 	private:
 		/**
@@ -157,7 +159,7 @@ namespace UE::PixelStreaming2
 
 		TWeakPtr<IPixelStreaming2DataProtocol> WeakDataProtocol;
 
-		bool Send(TArray<uint8>& Buffer) const;
+		UE_API bool Send(TArray<uint8>& Buffer) const;
 
 		template <typename... Args>
 		FBufferBuilder EncodeMessage(uint8 Type, Args... VarArgs) const
@@ -172,17 +174,19 @@ namespace UE::PixelStreaming2
 		}
 	};
 
-	class PIXELSTREAMING2RTC_API FEpicRtcMutliplexDataTrack : public FEpicRtcDataTrack
+	class FEpicRtcMutliplexDataTrack : public FEpicRtcDataTrack
 	{
 	public:
-		static TSharedPtr<FEpicRtcMutliplexDataTrack> Create(TSharedPtr<FEpicRtcDataTrack> InTrack, TWeakPtr<IPixelStreaming2DataProtocol> InDataProtocol, const FString& InPlayerId);
+		static UE_API TSharedPtr<FEpicRtcMutliplexDataTrack> Create(TSharedPtr<FEpicRtcDataTrack> InTrack, TWeakPtr<IPixelStreaming2DataProtocol> InDataProtocol, const FString& InPlayerId);
 
 	protected:
-		FEpicRtcMutliplexDataTrack(TSharedPtr<FEpicRtcDataTrack> InTrack, TWeakPtr<IPixelStreaming2DataProtocol> InDataProtocol, const FString& InPlayerId);
+		UE_API FEpicRtcMutliplexDataTrack(TSharedPtr<FEpicRtcDataTrack> InTrack, TWeakPtr<IPixelStreaming2DataProtocol> InDataProtocol, const FString& InPlayerId);
 
-		virtual void PrependData(FBufferBuilder& Builder) const override;
+		UE_API virtual void PrependData(FBufferBuilder& Builder) const override;
 
 	private:
 		const FString PlayerId;
 	};
-} // namespace UE::PixelStreaming2
\ No newline at end of file
+} // namespace UE::PixelStreaming2
+
+#undef UE_API
diff --git a/Engine/Plugins/Media/PixelStreaming2/Source/PixelStreaming2RTC/Internal/EpicRtcDataTrackObserver.h b/Engine/Plugins/Media/PixelStreaming2/Source/PixelStreaming2RTC/Internal/EpicRtcDataTrackObserver.h
--- a/Engine/Plugins/Media/PixelStreaming2/Source/PixelStreaming2RTC/Internal/EpicRtcDataTrackObserver.h
+++ b/Engine/Plugins/Media/PixelStreaming2/Source/PixelStreaming2RTC/Internal/EpicRtcDataTrackObserver.h
@@ -9,13 +9,15 @@
 
 #include "EpicRtcDataTrackObserver.generated.h"
 
+#define UE_API PIXELSTREAMING2RTC_API
+
 UINTERFACE(MinimalAPI)
 class UPixelStreaming2DataTrackObserver : public UInterface
 {
 	GENERATED_BODY()
 };
 
-class PIXELSTREAMING2RTC_API IPixelStreaming2DataTrackObserver
+class IPixelStreaming2DataTrackObserver
 {
 	GENERATED_BODY()
 
@@ -27,17 +29,17 @@ public:
 
 namespace UE::PixelStreaming2
 {
-	class PIXELSTREAMING2RTC_API FEpicRtcDataTrackObserver : public EpicRtcDataTrackObserverInterface
+	class FEpicRtcDataTrackObserver : public EpicRtcDataTrackObserverInterface
 	{
 	public:
-		FEpicRtcDataTrackObserver(TObserverVariant<IPixelStreaming2DataTrackObserver> UserObserver);
+		UE_API FEpicRtcDataTrackObserver(TObserverVariant<IPixelStreaming2DataTrackObserver> UserObserver);
 		virtual ~FEpicRtcDataTrackObserver() = default;
 
 	private:
 		// Begin EpicRtcDataTrackObserverInterface
-		virtual void OnDataTrackState(EpicRtcDataTrackInterface* DataTrack, const EpicRtcTrackState State) override;
-		virtual void OnDataTrackMessage(EpicRtcDataTrackInterface* DataTrack) override;
-		virtual void OnDataTrackError(EpicRtcDataTrackInterface*, const EpicRtcErrorCode) override;
+		UE_API virtual void OnDataTrackState(EpicRtcDataTrackInterface* DataTrack, const EpicRtcTrackState State) override;
+		UE_API virtual void OnDataTrackMessage(EpicRtcDataTrackInterface* DataTrack) override;
+		UE_API virtual void OnDataTrackError(EpicRtcDataTrackInterface*, const EpicRtcErrorCode) override;
 		// End EpicRtcDataTrackObserverInterface
 
 	public:
@@ -49,4 +51,6 @@ namespace UE::PixelStreaming2
 		TObserverVariant<IPixelStreaming2DataTrackObserver> UserObserver;
 	};
 
-} // namespace UE::PixelStreaming2
\ No newline at end of file
+} // namespace UE::PixelStreaming2
+
+#undef UE_API
diff --git a/Engine/Plugins/Media/PixelStreaming2/Source/PixelStreaming2RTC/Internal/EpicRtcDataTrackObserverFactory.h b/Engine/Plugins/Media/PixelStreaming2/Source/PixelStreaming2RTC/Internal/EpicRtcDataTrackObserverFactory.h
--- a/Engine/Plugins/Media/PixelStreaming2/Source/PixelStreaming2RTC/Internal/EpicRtcDataTrackObserverFactory.h
+++ b/Engine/Plugins/Media/PixelStreaming2/Source/PixelStreaming2RTC/Internal/EpicRtcDataTrackObserverFactory.h
@@ -7,19 +7,21 @@
 #include "epic_rtc/core/data_track_observer.h"
 #include "epic_rtc_helper/memory/ref_count_impl_helper.h"
 
+#define UE_API PIXELSTREAMING2RTC_API
+
 class IPixelStreaming2DataTrackObserver;
 
 namespace UE::PixelStreaming2
 {
-	class PIXELSTREAMING2RTC_API FEpicRtcDataTrackObserverFactory : public EpicRtcDataTrackObserverFactoryInterface
+	class FEpicRtcDataTrackObserverFactory : public EpicRtcDataTrackObserverFactoryInterface
 	{
 	public:
-		FEpicRtcDataTrackObserverFactory(TObserverVariant<IPixelStreaming2DataTrackObserver> UserObserver);
+		UE_API FEpicRtcDataTrackObserverFactory(TObserverVariant<IPixelStreaming2DataTrackObserver> UserObserver);
 		virtual ~FEpicRtcDataTrackObserverFactory() = default;
 
 	public:
 		// Begin EpicRtcDataTrackObserverFactoryInterface
-		virtual EpicRtcErrorCode CreateDataTrackObserver(const EpicRtcStringView ParticipantId, const EpicRtcStringView DataTrackId, EpicRtcDataTrackObserverInterface** OutDataTrackObserver) override;
+		UE_API virtual EpicRtcErrorCode CreateDataTrackObserver(const EpicRtcStringView ParticipantId, const EpicRtcStringView DataTrackId, EpicRtcDataTrackObserverInterface** OutDataTrackObserver) override;
 		// End EpicRtcDataTrackObserverFactoryInterface
 
 	public:
@@ -31,4 +33,6 @@ namespace UE::PixelStreaming2
 		TObserverVariant<IPixelStreaming2DataTrackObserver> UserObserver;
 	};
 
-} // namespace UE::PixelStreaming2
\ No newline at end of file
+} // namespace UE::PixelStreaming2
+
+#undef UE_API
diff --git a/Engine/Plugins/Media/PixelStreaming2/Source/PixelStreaming2RTC/Internal/EpicRtcLogging.h b/Engine/Plugins/Media/PixelStreaming2/Source/PixelStreaming2RTC/Internal/EpicRtcLogging.h
--- a/Engine/Plugins/Media/PixelStreaming2/Source/PixelStreaming2RTC/Internal/EpicRtcLogging.h
+++ b/Engine/Plugins/Media/PixelStreaming2/Source/PixelStreaming2RTC/Internal/EpicRtcLogging.h
@@ -8,6 +8,8 @@
 
 #include "epic_rtc/common/logging.h"
 
+#define UE_API PIXELSTREAMING2RTC_API
+
 PIXELSTREAMING2RTC_API DECLARE_LOG_CATEGORY_EXTERN(LogPixelStreaming2EpicRtc, Log, All);
 PIXELSTREAMING2RTC_API DECLARE_LOG_CATEGORY_EXTERN(LogPixelStreaming2WebRtc, Fatal, All);
 
@@ -41,7 +43,7 @@ namespace UE::PixelStreaming2
 	 * This includes completely removing log messages by checking IsFiltered()
 	 * or by redacting sensitive information using Censor()
 	 */
-	class PIXELSTREAMING2RTC_API ILogManipulator
+	class ILogManipulator
 	{
 	public:
 		virtual ~ILogManipulator() = default;
@@ -56,19 +58,19 @@ namespace UE::PixelStreaming2
 		virtual FString Censor(ELogVerbosity::Type LogVerbosity, const FString& LogString) = 0;
 	};
 
-	class PIXELSTREAMING2RTC_API FEpicRtcLogFilter : public ILogManipulator
+	class FEpicRtcLogFilter : public ILogManipulator
 	{
 	public:
-		FEpicRtcLogFilter();
-		~FEpicRtcLogFilter();
+		UE_API FEpicRtcLogFilter();
+		UE_API ~FEpicRtcLogFilter();
 
-		virtual bool IsFiltered(ELogVerbosity::Type LogVerbosity, const FString& LogString) override;
+		UE_API virtual bool IsFiltered(ELogVerbosity::Type LogVerbosity, const FString& LogString) override;
 
 		virtual FString Censor(ELogVerbosity::Type LogVerbosity, const FString& LogString) override { return LogString; }
 
 	private:
-		void OnEpicRtcLogFilterChanged(IConsoleVariable* Var);
-		void ParseFilterString(const FString& LogFilterString);
+		UE_API void OnEpicRtcLogFilterChanged(IConsoleVariable* Var);
+		UE_API void ParseFilterString(const FString& LogFilterString);
 
 	private:
 		TArray<FRegexPattern> RegexPatterns;
@@ -76,14 +78,16 @@ namespace UE::PixelStreaming2
 		FDelegateHandle EpicRtcLogFilterChangedHandle;
 	};
 
-	class PIXELSTREAMING2RTC_API FEpicRtcLogsRedirector : public EpicRtcLoggerInterface
+	class FEpicRtcLogsRedirector : public EpicRtcLoggerInterface
 	{
 	public:
-		FEpicRtcLogsRedirector(TSharedPtr<ILogManipulator> LogManipulator = nullptr);
+		UE_API FEpicRtcLogsRedirector(TSharedPtr<ILogManipulator> LogManipulator = nullptr);
 
-		virtual void Log(const EpicRtcLogMessage& Message) override;
+		UE_API virtual void Log(const EpicRtcLogMessage& Message) override;
 
 	private:
 		TSharedPtr<ILogManipulator> LogManipulator;
 	};
-} // namespace UE::PixelStreaming2
\ No newline at end of file
+} // namespace UE::PixelStreaming2
+
+#undef UE_API
diff --git a/Engine/Plugins/Media/PixelStreaming2/Source/PixelStreaming2RTC/Internal/EpicRtcRoomObserver.h b/Engine/Plugins/Media/PixelStreaming2/Source/PixelStreaming2RTC/Internal/EpicRtcRoomObserver.h
--- a/Engine/Plugins/Media/PixelStreaming2/Source/PixelStreaming2RTC/Internal/EpicRtcRoomObserver.h
+++ b/Engine/Plugins/Media/PixelStreaming2/Source/PixelStreaming2RTC/Internal/EpicRtcRoomObserver.h
@@ -9,13 +9,15 @@
 
 #include "EpicRtcRoomObserver.generated.h"
 
+#define UE_API PIXELSTREAMING2RTC_API
+
 UINTERFACE(MinimalAPI)
 class UPixelStreaming2RoomObserver : public UInterface
 {
 	GENERATED_BODY()
 };
 
-class PIXELSTREAMING2RTC_API IPixelStreaming2RoomObserver
+class IPixelStreaming2RoomObserver
 {
 	GENERATED_BODY()
 
@@ -33,23 +35,23 @@ public:
 
 namespace UE::PixelStreaming2
 {
-	class PIXELSTREAMING2RTC_API FEpicRtcRoomObserver : public EpicRtcRoomObserverInterface
+	class FEpicRtcRoomObserver : public EpicRtcRoomObserverInterface
 	{
 	public:
-		FEpicRtcRoomObserver(TObserverVariant<IPixelStreaming2RoomObserver> UserObserver);
+		UE_API FEpicRtcRoomObserver(TObserverVariant<IPixelStreaming2RoomObserver> UserObserver);
 		virtual ~FEpicRtcRoomObserver() = default;
 
 	private:
 		// Begin EpicRtcRoomObserver
-		virtual void							   OnRoomStateUpdate(const EpicRtcRoomState State) override;
-		virtual void							   OnRoomJoinedUpdate(EpicRtcParticipantInterface* Participant) override;
-		virtual void							   OnRoomLeftUpdate(const EpicRtcStringView ParticipantId) override;
-		virtual void							   OnAudioTrackUpdate(EpicRtcParticipantInterface* Participant, EpicRtcAudioTrackInterface* AudioTrack) override;
-		virtual void							   OnVideoTrackUpdate(EpicRtcParticipantInterface* Participant, EpicRtcVideoTrackInterface* VideoTrack) override;
-		virtual void							   OnDataTrackUpdate(EpicRtcParticipantInterface* Participant, EpicRtcDataTrackInterface* DataTrack) override;
-		[[nodiscard]] virtual EpicRtcSdpInterface* OnLocalSdpUpdate(EpicRtcParticipantInterface* Participant, EpicRtcSdpInterface* Sdp) override;
-		[[nodiscard]] virtual EpicRtcSdpInterface* OnRemoteSdpUpdate(EpicRtcParticipantInterface* Participant, EpicRtcSdpInterface* Sdp) override;
-		virtual void							   OnRoomErrorUpdate(const EpicRtcErrorCode Error) override;
+		UE_API virtual void							   OnRoomStateUpdate(const EpicRtcRoomState State) override;
+		UE_API virtual void							   OnRoomJoinedUpdate(EpicRtcParticipantInterface* Participant) override;
+		UE_API virtual void							   OnRoomLeftUpdate(const EpicRtcStringView ParticipantId) override;
+		UE_API virtual void							   OnAudioTrackUpdate(EpicRtcParticipantInterface* Participant, EpicRtcAudioTrackInterface* AudioTrack) override;
+		UE_API virtual void							   OnVideoTrackUpdate(EpicRtcParticipantInterface* Participant, EpicRtcVideoTrackInterface* VideoTrack) override;
+		UE_API virtual void							   OnDataTrackUpdate(EpicRtcParticipantInterface* Participant, EpicRtcDataTrackInterface* DataTrack) override;
+		[[nodiscard]] UE_API virtual EpicRtcSdpInterface* OnLocalSdpUpdate(EpicRtcParticipantInterface* Participant, EpicRtcSdpInterface* Sdp) override;
+		[[nodiscard]] UE_API virtual EpicRtcSdpInterface* OnRemoteSdpUpdate(EpicRtcParticipantInterface* Participant, EpicRtcSdpInterface* Sdp) override;
+		UE_API virtual void							   OnRoomErrorUpdate(const EpicRtcErrorCode Error) override;
 		// Begin EpicRtcRoomObserver
 
 	public:
@@ -61,4 +63,6 @@ namespace UE::PixelStreaming2
 		TObserverVariant<IPixelStreaming2RoomObserver> UserObserver;
 	};
 
-} // namespace UE::PixelStreaming2
\ No newline at end of file
+} // namespace UE::PixelStreaming2
+
+#undef UE_API
diff --git a/Engine/Plugins/Media/PixelStreaming2/Source/PixelStreaming2RTC/Internal/EpicRtcSessionObserver.h b/Engine/Plugins/Media/PixelStreaming2/Source/PixelStreaming2RTC/Internal/EpicRtcSessionObserver.h
--- a/Engine/Plugins/Media/PixelStreaming2/Source/PixelStreaming2RTC/Internal/EpicRtcSessionObserver.h
+++ b/Engine/Plugins/Media/PixelStreaming2/Source/PixelStreaming2RTC/Internal/EpicRtcSessionObserver.h
@@ -9,13 +9,15 @@
 
 #include "EpicRtcSessionObserver.generated.h"
 
+#define UE_API PIXELSTREAMING2RTC_API
+
 UINTERFACE(MinimalAPI)
 class UPixelStreaming2SessionObserver : public UInterface
 {
 	GENERATED_BODY()
 };
 
-class PIXELSTREAMING2RTC_API IPixelStreaming2SessionObserver
+class IPixelStreaming2SessionObserver
 {
 	GENERATED_BODY()
 
@@ -27,17 +29,17 @@ public:
 
 namespace UE::PixelStreaming2
 {
-	class PIXELSTREAMING2RTC_API FEpicRtcSessionObserver : public EpicRtcSessionObserverInterface
+	class FEpicRtcSessionObserver : public EpicRtcSessionObserverInterface
 	{
 	public:
-		FEpicRtcSessionObserver(TObserverVariant<IPixelStreaming2SessionObserver> UserObserver);
+		UE_API FEpicRtcSessionObserver(TObserverVariant<IPixelStreaming2SessionObserver> UserObserver);
 		virtual ~FEpicRtcSessionObserver() = default;
 
 	private:
 		// Begin EpicRtcSessionObserver
-		virtual void OnSessionStateUpdate(const EpicRtcSessionState State) override;
-		virtual void OnSessionErrorUpdate(const EpicRtcErrorCode Error) override;
-		virtual void OnSessionRoomsAvailableUpdate(EpicRtcStringArrayInterface* RoomsList) override;
+		UE_API virtual void OnSessionStateUpdate(const EpicRtcSessionState State) override;
+		UE_API virtual void OnSessionErrorUpdate(const EpicRtcErrorCode Error) override;
+		UE_API virtual void OnSessionRoomsAvailableUpdate(EpicRtcStringArrayInterface* RoomsList) override;
 		// End EpicRtcSessionObserver
 	public:
 		// Begin EpicRtcRefCountInterface
@@ -48,3 +50,5 @@ namespace UE::PixelStreaming2
 	};
 
 } // namespace UE::PixelStreaming2
+
+#undef UE_API
diff --git a/Engine/Plugins/Media/PixelStreaming2/Source/PixelStreaming2RTC/Internal/EpicRtcStatsCollector.h b/Engine/Plugins/Media/PixelStreaming2/Source/PixelStreaming2RTC/Internal/EpicRtcStatsCollector.h
--- a/Engine/Plugins/Media/PixelStreaming2/Source/PixelStreaming2RTC/Internal/EpicRtcStatsCollector.h
+++ b/Engine/Plugins/Media/PixelStreaming2/Source/PixelStreaming2RTC/Internal/EpicRtcStatsCollector.h
@@ -9,16 +9,18 @@
 #include "epic_rtc/core/stats.h"
 #include "epic_rtc_helper/memory/ref_count_impl_helper.h"
 
+#define UE_API PIXELSTREAMING2RTC_API
+
 namespace UE::PixelStreaming2
 {
-	class PIXELSTREAMING2RTC_API FEpicRtcStatsCollector : public EpicRtcStatsCollectorCallbackInterface
+	class FEpicRtcStatsCollector : public EpicRtcStatsCollectorCallbackInterface
 	{
 	public:
 		FEpicRtcStatsCollector() = default;
 		~FEpicRtcStatsCollector() = default;
 
 		// Begin EpicRtcStatsCollectorCallbackInterface interface
-		void			 OnStatsDelivered(const EpicRtcStatsReport& InReport) override;
+		UE_API void			 OnStatsDelivered(const EpicRtcStatsReport& InReport) override;
 		// End EpicRtcStatsCollectorCallbackInterface interface
 
 		// Begin EpicRtcRefCountInterface
@@ -29,4 +31,6 @@ namespace UE::PixelStreaming2
 		DECLARE_EVENT_TwoParams(FEpicRtcStatsCollector, FOnStatsReady, const FString&, const EpicRtcConnectionStats&);
 		FOnStatsReady OnStatsReady;
 	};
-} // namespace UE::PixelStreaming2
\ No newline at end of file
+} // namespace UE::PixelStreaming2
+
+#undef UE_API
diff --git a/Engine/Plugins/Media/PixelStreaming2/Source/PixelStreaming2RTC/Internal/EpicRtcVideoCapturer.h b/Engine/Plugins/Media/PixelStreaming2/Source/PixelStreaming2RTC/Internal/EpicRtcVideoCapturer.h
--- a/Engine/Plugins/Media/PixelStreaming2/Source/PixelStreaming2RTC/Internal/EpicRtcVideoCapturer.h
+++ b/Engine/Plugins/Media/PixelStreaming2/Source/PixelStreaming2RTC/Internal/EpicRtcVideoCapturer.h
@@ -6,17 +6,21 @@
 
 #include "epic_rtc/core/video/video_buffer.h"
 
+#define UE_API PIXELSTREAMING2RTC_API
+
 namespace UE::PixelStreaming2
 {
-	class PIXELSTREAMING2RTC_API FEpicRtcVideoCapturer : public FVideoCapturer
+	class FEpicRtcVideoCapturer : public FVideoCapturer
 	{
 	public:
-		static TSharedPtr<FEpicRtcVideoCapturer> Create(TSharedPtr<FVideoProducer> VideoProducer = nullptr, bool bIsSRGB = false);
+		static UE_API TSharedPtr<FEpicRtcVideoCapturer> Create(TSharedPtr<FVideoProducer> VideoProducer = nullptr, bool bIsSRGB = false);
 		virtual ~FEpicRtcVideoCapturer() = default;
 
-		TRefCountPtr<EpicRtcVideoBufferInterface> GetFrameBuffer();
+		UE_API TRefCountPtr<EpicRtcVideoBufferInterface> GetFrameBuffer();
 
 	private:
-		FEpicRtcVideoCapturer(TSharedPtr<FVideoProducer> VideoProducer, bool bIsSRGB);
+		UE_API FEpicRtcVideoCapturer(TSharedPtr<FVideoProducer> VideoProducer, bool bIsSRGB);
 	};
 } // namespace UE::PixelStreaming2
+
+#undef UE_API
diff --git a/Engine/Plugins/Media/PixelStreaming2/Source/PixelStreaming2RTC/Internal/EpicRtcVideoCommon.h b/Engine/Plugins/Media/PixelStreaming2/Source/PixelStreaming2RTC/Internal/EpicRtcVideoCommon.h
--- a/Engine/Plugins/Media/PixelStreaming2/Source/PixelStreaming2RTC/Internal/EpicRtcVideoCommon.h
+++ b/Engine/Plugins/Media/PixelStreaming2/Source/PixelStreaming2RTC/Internal/EpicRtcVideoCommon.h
@@ -17,7 +17,7 @@ FORCEINLINE bool operator==(const EpicRtcVideoResolution& Lhs, const EpicRtcVide
 
 namespace UE::PixelStreaming2
 {
-	class PIXELSTREAMING2RTC_API FEpicRtcString : public EpicRtcStringInterface
+	class FEpicRtcString : public EpicRtcStringInterface
 	{
 	public:
 		FEpicRtcString() = default;
@@ -45,7 +45,7 @@ namespace UE::PixelStreaming2
 		// End EpicRtcRefCountInterface
 	};
 
-	class PIXELSTREAMING2RTC_API FEpicRtcEncodedVideoBuffer : public EpicRtcEncodedVideoBufferInterface
+	class FEpicRtcEncodedVideoBuffer : public EpicRtcEncodedVideoBufferInterface
 	{
 	public:
 		FEpicRtcEncodedVideoBuffer() = default;
@@ -67,7 +67,7 @@ namespace UE::PixelStreaming2
 		// End EpicRtcRefCountInterface
 	};
 
-	class PIXELSTREAMING2RTC_API FEpicRtcParameterPairArray : public EpicRtcParameterPairArrayInterface
+	class FEpicRtcParameterPairArray : public EpicRtcParameterPairArrayInterface
 	{
 	public:
 		FEpicRtcParameterPairArray() = default;
@@ -112,7 +112,7 @@ namespace UE::PixelStreaming2
 		// End EpicRtcRefCountInterface
 	};
 
-	class PIXELSTREAMING2RTC_API FEpicRtcParameterPair : public EpicRtcParameterPairInterface
+	class FEpicRtcParameterPair : public EpicRtcParameterPairInterface
 	{
 	public:
 		FEpicRtcParameterPair(EpicRtcStringInterface* Key, EpicRtcStringInterface* Value)
@@ -141,7 +141,7 @@ namespace UE::PixelStreaming2
 		// End EpicRtcRefCountInterface
 	};
 
-	class PIXELSTREAMING2RTC_API FEpicRtcVideoParameterPairArray : public EpicRtcVideoParameterPairArrayInterface
+	class FEpicRtcVideoParameterPairArray : public EpicRtcVideoParameterPairArrayInterface
 	{
 	public:
 		FEpicRtcVideoParameterPairArray() = default;
@@ -229,7 +229,7 @@ namespace UE::PixelStreaming2
 		// End EpicRtcRefCountInterface
 	};
 
-	class PIXELSTREAMING2RTC_API FEpicRtcScalabilityModeArray : public EpicRtcVideoScalabilityModeArrayInterface
+	class FEpicRtcScalabilityModeArray : public EpicRtcVideoScalabilityModeArrayInterface
 	{
 	public:
 		FEpicRtcScalabilityModeArray() = default;
@@ -282,7 +282,7 @@ namespace UE::PixelStreaming2
 		// End EpicRtcRefCountInterface
 	};
 
-	class PIXELSTREAMING2RTC_API FEpicRtcVideoCodecInfo : public EpicRtcVideoCodecInfoInterface
+	class FEpicRtcVideoCodecInfo : public EpicRtcVideoCodecInfoInterface
 	{
 	public:
 		FEpicRtcVideoCodecInfo(EpicRtcVideoCodec Codec, bool bIsHardwareAccelerated = false, EpicRtcVideoParameterPairArrayInterface* Parameters = new FEpicRtcVideoParameterPairArray(), EpicRtcVideoScalabilityModeArrayInterface* ScalabilityModes = new FEpicRtcScalabilityModeArray())
@@ -327,7 +327,7 @@ namespace UE::PixelStreaming2
 		// End EpicRtcRefCountInterface
 	};
 
-	class PIXELSTREAMING2RTC_API FVideoCodecInfoArray : public EpicRtcVideoCodecInfoArrayInterface
+	class FVideoCodecInfoArray : public EpicRtcVideoCodecInfoArrayInterface
 	{
 	public:
 		FVideoCodecInfoArray() = default;
@@ -403,7 +403,7 @@ namespace UE::PixelStreaming2
 		// End EpicRtcRefCountInterface
 	};
 
-	class PIXELSTREAMING2RTC_API FEpicRtcVideoResolutionBitrateLimitsArray : public EpicRtcVideoResolutionBitrateLimitsArrayInterface
+	class FEpicRtcVideoResolutionBitrateLimitsArray : public EpicRtcVideoResolutionBitrateLimitsArrayInterface
 	{
 	public:
 		FEpicRtcVideoResolutionBitrateLimitsArray() = default;
@@ -443,7 +443,7 @@ namespace UE::PixelStreaming2
 		// End EpicRtcRefCountInterface
 	};
 
-	class PIXELSTREAMING2RTC_API FEpicRtcPixelFormatArray : public EpicRtcPixelFormatArrayInterface
+	class FEpicRtcPixelFormatArray : public EpicRtcPixelFormatArrayInterface
 	{
 	public:
 		FEpicRtcPixelFormatArray() = default;
@@ -483,7 +483,7 @@ namespace UE::PixelStreaming2
 		// End EpicRtcRefCountInterface
 	};
 
-	class PIXELSTREAMING2RTC_API FEpicRtcVideoFrameTypeArray : public EpicRtcVideoFrameTypeArrayInterface
+	class FEpicRtcVideoFrameTypeArray : public EpicRtcVideoFrameTypeArrayInterface
 	{
 	public:
 		FEpicRtcVideoFrameTypeArray() = default;
@@ -518,7 +518,7 @@ namespace UE::PixelStreaming2
 		// End EpicRtcRefCountInterface
 	};
 
-	class PIXELSTREAMING2RTC_API FEpicRtcInt32Array : public EpicRtcInt32ArrayInterface
+	class FEpicRtcInt32Array : public EpicRtcInt32ArrayInterface
 	{
 	public:
 		FEpicRtcInt32Array() = default;
@@ -563,7 +563,7 @@ namespace UE::PixelStreaming2
 		// End EpicRtcRefCountInterface
 	};
 
-	class PIXELSTREAMING2RTC_API FEpicRtcBoolArray : public EpicRtcBoolArrayInterface
+	class FEpicRtcBoolArray : public EpicRtcBoolArrayInterface
 	{
 	public:
 		FEpicRtcBoolArray() = default;
@@ -620,7 +620,7 @@ namespace UE::PixelStreaming2
 		// End EpicRtcRefCountInterface
 	};
 
-	class PIXELSTREAMING2RTC_API FEpicRtcDecodeTargetIndicationArray : public EpicRtcDecodeTargetIndicationArrayInterface
+	class FEpicRtcDecodeTargetIndicationArray : public EpicRtcDecodeTargetIndicationArrayInterface
 	{
 	public:
 		FEpicRtcDecodeTargetIndicationArray() = default;
@@ -687,7 +687,7 @@ namespace UE::PixelStreaming2
 		// End EpicRtcRefCountInterface
 	};
 
-	class PIXELSTREAMING2RTC_API FEpicRtcCodecBufferUsageArray : public EpicRtcCodecBufferUsageArrayInterface
+	class FEpicRtcCodecBufferUsageArray : public EpicRtcCodecBufferUsageArrayInterface
 	{
 	public:
 		FEpicRtcCodecBufferUsageArray() = default;
@@ -742,7 +742,7 @@ namespace UE::PixelStreaming2
 		// End EpicRtcRefCountInterface
 	};
 
-	class PIXELSTREAMING2RTC_API FEpicRtcVideoResolutionArray : public EpicRtcVideoResolutionArrayInterface
+	class FEpicRtcVideoResolutionArray : public EpicRtcVideoResolutionArrayInterface
 	{
 	public:
 		FEpicRtcVideoResolutionArray() = default;
@@ -796,7 +796,7 @@ namespace UE::PixelStreaming2
 		// End EpicRtcRefCountInterface
 	};
 
-	class PIXELSTREAMING2RTC_API FEpicRtcGenericFrameInfoArray : public EpicRtcGenericFrameInfoArrayInterface
+	class FEpicRtcGenericFrameInfoArray : public EpicRtcGenericFrameInfoArrayInterface
 	{
 	public:
 		FEpicRtcGenericFrameInfoArray() = default;
@@ -872,7 +872,7 @@ namespace UE::PixelStreaming2
 		// End EpicRtcRefCountInterface
 	};
 
-	class PIXELSTREAMING2RTC_API FEpicRtcGenericFrameInfo : public EpicRtcGenericFrameInfoInterface
+	class FEpicRtcGenericFrameInfo : public EpicRtcGenericFrameInfoInterface
 	{
 	public:
 		FEpicRtcGenericFrameInfo(const FGenericFrameInfo& GenericFrameInfo)
@@ -914,7 +914,7 @@ namespace UE::PixelStreaming2
 		// End EpicRtcRefCountInterface
 	};
 
-	class PIXELSTREAMING2RTC_API FEpicRtcFrameDependencyStructure : public EpicRtcFrameDependencyStructure
+	class FEpicRtcFrameDependencyStructure : public EpicRtcFrameDependencyStructure
 	{
 	public:
 		FEpicRtcFrameDependencyStructure(const FFrameDependencyStructure& FrameDependencyStructure)
diff --git a/Engine/Plugins/Media/PixelStreaming2/Source/PixelStreaming2RTC/Internal/EpicRtcVideoDecoderInitializer.h b/Engine/Plugins/Media/PixelStreaming2/Source/PixelStreaming2RTC/Internal/EpicRtcVideoDecoderInitializer.h
--- a/Engine/Plugins/Media/PixelStreaming2/Source/PixelStreaming2RTC/Internal/EpicRtcVideoDecoderInitializer.h
+++ b/Engine/Plugins/Media/PixelStreaming2/Source/PixelStreaming2RTC/Internal/EpicRtcVideoDecoderInitializer.h
@@ -8,23 +8,25 @@
 #include "epic_rtc/core/video/video_decoder.h"
 #include "epic_rtc_helper/memory/ref_count_impl_helper.h"
 
+#define UE_API PIXELSTREAMING2RTC_API
+
 namespace UE::PixelStreaming2
 {
 
-	class PIXELSTREAMING2RTC_API FEpicRtcVideoDecoderInitializer : public EpicRtcVideoDecoderInitializerInterface
+	class FEpicRtcVideoDecoderInitializer : public EpicRtcVideoDecoderInitializerInterface
 	{
 	public:
 		FEpicRtcVideoDecoderInitializer() = default;
 		virtual ~FEpicRtcVideoDecoderInitializer() = default;
 
 		// Begin EpicRtcVideoDecoderInitializerInterface
-		virtual void								 CreateDecoder(EpicRtcVideoCodecInfoInterface* CodecInfo, EpicRtcVideoDecoderInterface** OutDecoder) override;
-		virtual EpicRtcStringView					 GetName() override;
-		virtual EpicRtcVideoCodecInfoArrayInterface* GetSupportedCodecs() override;
+		UE_API virtual void								 CreateDecoder(EpicRtcVideoCodecInfoInterface* CodecInfo, EpicRtcVideoDecoderInterface** OutDecoder) override;
+		UE_API virtual EpicRtcStringView					 GetName() override;
+		UE_API virtual EpicRtcVideoCodecInfoArrayInterface* GetSupportedCodecs() override;
 		// End EpicRtcVideoDecoderInitializerInterface
 
 	private:
-		TMap<EVideoCodec, TArray<TRefCountPtr<EpicRtcVideoCodecInfoInterface>>> CreateSupportedDecoderMap();
+		UE_API TMap<EVideoCodec, TArray<TRefCountPtr<EpicRtcVideoCodecInfoInterface>>> CreateSupportedDecoderMap();
 
 	public:
 		// Begin EpicRtcRefCountInterface
@@ -32,4 +34,6 @@ namespace UE::PixelStreaming2
 		// End EpicRtcRefCountInterface
 	};
 
-} // namespace UE::PixelStreaming2
\ No newline at end of file
+} // namespace UE::PixelStreaming2
+
+#undef UE_API
diff --git a/Engine/Plugins/Media/PixelStreaming2/Source/PixelStreaming2RTC/Internal/EpicRtcVideoEncoderInitializer.h b/Engine/Plugins/Media/PixelStreaming2/Source/PixelStreaming2RTC/Internal/EpicRtcVideoEncoderInitializer.h
--- a/Engine/Plugins/Media/PixelStreaming2/Source/PixelStreaming2RTC/Internal/EpicRtcVideoEncoderInitializer.h
+++ b/Engine/Plugins/Media/PixelStreaming2/Source/PixelStreaming2RTC/Internal/EpicRtcVideoEncoderInitializer.h
@@ -8,23 +8,25 @@
 #include "epic_rtc/core/video/video_encoder.h"
 #include "epic_rtc_helper/memory/ref_count_impl_helper.h"
 
+#define UE_API PIXELSTREAMING2RTC_API
+
 namespace UE::PixelStreaming2
 {
 
-	class PIXELSTREAMING2RTC_API FEpicRtcVideoEncoderInitializer : public EpicRtcVideoEncoderInitializerInterface
+	class FEpicRtcVideoEncoderInitializer : public EpicRtcVideoEncoderInitializerInterface
 	{
 	public:
 		FEpicRtcVideoEncoderInitializer() = default;
 		virtual ~FEpicRtcVideoEncoderInitializer() = default;
 
 		// Begin EpicRtcVideoEncoderInitializerInterface
-		virtual void								 CreateEncoder(EpicRtcVideoCodecInfoInterface* CodecInfo, EpicRtcVideoEncoderInterface** OutEncoder) override;
-		virtual EpicRtcStringView					 GetName() override;
-		virtual EpicRtcVideoCodecInfoArrayInterface* GetSupportedCodecs() override;
+		UE_API virtual void								 CreateEncoder(EpicRtcVideoCodecInfoInterface* CodecInfo, EpicRtcVideoEncoderInterface** OutEncoder) override;
+		UE_API virtual EpicRtcStringView					 GetName() override;
+		UE_API virtual EpicRtcVideoCodecInfoArrayInterface* GetSupportedCodecs() override;
 		// End EpicRtcVideoEncoderInitializerInterface
 
 	private:
-		TMap<EVideoCodec, TArray<TRefCountPtr<EpicRtcVideoCodecInfoInterface>>> CreateSupportedEncoderMap();
+		UE_API TMap<EVideoCodec, TArray<TRefCountPtr<EpicRtcVideoCodecInfoInterface>>> CreateSupportedEncoderMap();
 
 	public:
 		// Begin EpicRtcRefCountInterface
@@ -32,4 +34,6 @@ namespace UE::PixelStreaming2
 		// End EpicRtcRefCountInterface
 	};
 
-} // namespace UE::PixelStreaming2
\ No newline at end of file
+} // namespace UE::PixelStreaming2
+
+#undef UE_API
diff --git a/Engine/Plugins/Media/PixelStreaming2/Source/PixelStreaming2RTC/Internal/EpicRtcVideoSink.h b/Engine/Plugins/Media/PixelStreaming2/Source/PixelStreaming2RTC/Internal/EpicRtcVideoSink.h
--- a/Engine/Plugins/Media/PixelStreaming2/Source/PixelStreaming2RTC/Internal/EpicRtcVideoSink.h
+++ b/Engine/Plugins/Media/PixelStreaming2/Source/PixelStreaming2RTC/Internal/EpicRtcVideoSink.h
@@ -11,25 +11,27 @@
 
 #include "epic_rtc/core/video/video_frame.h"
 
+#define UE_API PIXELSTREAMING2RTC_API
+
 namespace UE::PixelStreaming2
 {
 	/**
 	 * Video sink class that receives a frame from EpicRtc and passes the frame to all added consumers
 	 */
-	class PIXELSTREAMING2RTC_API FEpicRtcVideoSink : public FVideoSink, public TEpicRtcTrack<EpicRtcVideoTrackInterface>, public TSharedFromThis<FEpicRtcVideoSink>
+	class FEpicRtcVideoSink : public FVideoSink, public TEpicRtcTrack<EpicRtcVideoTrackInterface>, public TSharedFromThis<FEpicRtcVideoSink>
 	{
 	public:
-		static TSharedPtr<FEpicRtcVideoSink> Create(TRefCountPtr<EpicRtcVideoTrackInterface> InTrack);
+		static UE_API TSharedPtr<FEpicRtcVideoSink> Create(TRefCountPtr<EpicRtcVideoTrackInterface> InTrack);
 		// Note: destructor will call destroy on any attached video consumers
 		virtual ~FEpicRtcVideoSink() = default;
 
-		void OnEpicRtcFrame(const EpicRtcVideoFrame& Frame);
+		UE_API void OnEpicRtcFrame(const EpicRtcVideoFrame& Frame);
 
 		inline bool TriggerDisconnectFlag() { return bDisconnectFlag.Trigger(); };
 	private:
-		FEpicRtcVideoSink(TRefCountPtr<EpicRtcVideoTrackInterface> InTrack);
+		UE_API FEpicRtcVideoSink(TRefCountPtr<EpicRtcVideoTrackInterface> InTrack);
 
-		void OnFrameCaptured();
+		UE_API void OnFrameCaptured();
 
 		FCriticalSection				  RenderSyncContext;
 		FPooledRenderTargetDesc			  RenderTargetDescriptor;
@@ -40,4 +42,6 @@ namespace UE::PixelStreaming2
 		TSharedPtr<FVideoCapturer> VideoCapturer;
 		FAtomicOnceFlag			   bDisconnectFlag;
 	};
-} // namespace UE::PixelStreaming2
\ No newline at end of file
+} // namespace UE::PixelStreaming2
+
+#undef UE_API
diff --git a/Engine/Plugins/Media/PixelStreaming2/Source/PixelStreaming2RTC/Internal/EpicRtcVideoSource.h b/Engine/Plugins/Media/PixelStreaming2/Source/PixelStreaming2RTC/Internal/EpicRtcVideoSource.h
--- a/Engine/Plugins/Media/PixelStreaming2/Source/PixelStreaming2RTC/Internal/EpicRtcVideoSource.h
+++ b/Engine/Plugins/Media/PixelStreaming2/Source/PixelStreaming2RTC/Internal/EpicRtcVideoSource.h
@@ -7,6 +7,8 @@
 #include "EpicRtcVideoCapturer.h"
 #include "VideoSource.h"
 
+#define UE_API PIXELSTREAMING2RTC_API
+
 namespace UE::PixelStreaming2
 {
 	class FVideoSourceGroup;
@@ -21,23 +23,25 @@ namespace UE::PixelStreaming2
 	 * will make sure that the adapter has valid output and if so will create a frame
 	 * for EpicRtc. Otherwise it will continue to wait until the next frame.
 	 */
-	class PIXELSTREAMING2RTC_API FEpicRtcVideoSource : public FVideoSource, public TEpicRtcTrack<EpicRtcVideoTrackInterface>
+	class FEpicRtcVideoSource : public FVideoSource, public TEpicRtcTrack<EpicRtcVideoTrackInterface>
 	{
 	public:
-		static TSharedPtr<FEpicRtcVideoSource> Create(TRefCountPtr<EpicRtcVideoTrackInterface> InVideoTrack, TSharedPtr<FEpicRtcVideoCapturer> InVideoCapturer, TSharedPtr<FVideoSourceGroup> InVideoSourceGroup);
+		static UE_API TSharedPtr<FEpicRtcVideoSource> Create(TRefCountPtr<EpicRtcVideoTrackInterface> InVideoTrack, TSharedPtr<FEpicRtcVideoCapturer> InVideoCapturer, TSharedPtr<FVideoSourceGroup> InVideoSourceGroup);
 		virtual ~FEpicRtcVideoSource() = default;
 
-		virtual void PushFrame() override;
-		virtual void ForceKeyFrame() override;
+		UE_API virtual void PushFrame() override;
+		UE_API virtual void ForceKeyFrame() override;
 
-		TRefCountPtr<EpicRtcVideoTrackInterface> GetVideoTrack();
+		UE_API TRefCountPtr<EpicRtcVideoTrackInterface> GetVideoTrack();
 		
 		inline bool TriggerDisconnectFlag() { return bDisconnectFlag.Trigger(); };
 	private:
-		FEpicRtcVideoSource(TRefCountPtr<EpicRtcVideoTrackInterface> InVideoTrack, TSharedPtr<FEpicRtcVideoCapturer> InVideoCapturer);
+		UE_API FEpicRtcVideoSource(TRefCountPtr<EpicRtcVideoTrackInterface> InVideoTrack, TSharedPtr<FEpicRtcVideoCapturer> InVideoCapturer);
 
 	private:
 		TSharedPtr<FEpicRtcVideoCapturer> VideoCapturer;
 		FAtomicOnceFlag					  bDisconnectFlag;
 	};
 } // namespace UE::PixelStreaming2
+
+#undef UE_API
diff --git a/Engine/Plugins/Media/PixelStreaming2/Source/PixelStreaming2RTC/Internal/EpicRtcVideoTrackObserver.h b/Engine/Plugins/Media/PixelStreaming2/Source/PixelStreaming2RTC/Internal/EpicRtcVideoTrackObserver.h
--- a/Engine/Plugins/Media/PixelStreaming2/Source/PixelStreaming2RTC/Internal/EpicRtcVideoTrackObserver.h
+++ b/Engine/Plugins/Media/PixelStreaming2/Source/PixelStreaming2RTC/Internal/EpicRtcVideoTrackObserver.h
@@ -9,13 +9,15 @@
 
 #include "EpicRtcVideoTrackObserver.generated.h"
 
+#define UE_API PIXELSTREAMING2RTC_API
+
 UINTERFACE(MinimalAPI)
 class UPixelStreaming2VideoTrackObserver : public UInterface
 {
 	GENERATED_BODY()
 };
 
-class PIXELSTREAMING2RTC_API IPixelStreaming2VideoTrackObserver
+class IPixelStreaming2VideoTrackObserver
 {
 	GENERATED_BODY()
 
@@ -30,20 +32,20 @@ public:
 
 namespace UE::PixelStreaming2
 {
-	class PIXELSTREAMING2RTC_API FEpicRtcVideoTrackObserver : public EpicRtcVideoTrackObserverInterface
+	class FEpicRtcVideoTrackObserver : public EpicRtcVideoTrackObserverInterface
 	{
 	public:
-		FEpicRtcVideoTrackObserver(TObserverVariant<IPixelStreaming2VideoTrackObserver> UserObserver);
+		UE_API FEpicRtcVideoTrackObserver(TObserverVariant<IPixelStreaming2VideoTrackObserver> UserObserver);
 		virtual ~FEpicRtcVideoTrackObserver() = default;
 
 	private:
 		// Begin EpicRtcVideoTrackObserverInterface
-		virtual void		OnVideoTrackMuted(EpicRtcVideoTrackInterface* VideoTrack, EpicRtcBool bIsMuted) override;
-		virtual void		OnVideoTrackFrame(EpicRtcVideoTrackInterface* VideoTrack, const EpicRtcVideoFrame& Frame) override;
-		virtual void		OnVideoTrackRemoved(EpicRtcVideoTrackInterface* VideoTrack) override;
-		virtual void		OnVideoTrackState(EpicRtcVideoTrackInterface* VideoTrack, const EpicRtcTrackState State) override;
-		virtual void		OnVideoTrackEncodedFrame(EpicRtcVideoTrackInterface*, const EpicRtcEncodedVideoFrame&) override;
-		virtual EpicRtcBool Enabled() const override;
+		UE_API virtual void		OnVideoTrackMuted(EpicRtcVideoTrackInterface* VideoTrack, EpicRtcBool bIsMuted) override;
+		UE_API virtual void		OnVideoTrackFrame(EpicRtcVideoTrackInterface* VideoTrack, const EpicRtcVideoFrame& Frame) override;
+		UE_API virtual void		OnVideoTrackRemoved(EpicRtcVideoTrackInterface* VideoTrack) override;
+		UE_API virtual void		OnVideoTrackState(EpicRtcVideoTrackInterface* VideoTrack, const EpicRtcTrackState State) override;
+		UE_API virtual void		OnVideoTrackEncodedFrame(EpicRtcVideoTrackInterface*, const EpicRtcEncodedVideoFrame&) override;
+		UE_API virtual EpicRtcBool Enabled() const override;
 		// End EpicRtcVideoTrackObserverInterface
 
 	public:
@@ -55,4 +57,6 @@ namespace UE::PixelStreaming2
 		TObserverVariant<IPixelStreaming2VideoTrackObserver> UserObserver;
 	};
 
-} // namespace UE::PixelStreaming2
\ No newline at end of file
+} // namespace UE::PixelStreaming2
+
+#undef UE_API
diff --git a/Engine/Plugins/Media/PixelStreaming2/Source/PixelStreaming2RTC/Internal/EpicRtcVideoTrackObserverFactory.h b/Engine/Plugins/Media/PixelStreaming2/Source/PixelStreaming2RTC/Internal/EpicRtcVideoTrackObserverFactory.h
--- a/Engine/Plugins/Media/PixelStreaming2/Source/PixelStreaming2RTC/Internal/EpicRtcVideoTrackObserverFactory.h
+++ b/Engine/Plugins/Media/PixelStreaming2/Source/PixelStreaming2RTC/Internal/EpicRtcVideoTrackObserverFactory.h
@@ -7,19 +7,21 @@
 #include "epic_rtc/core/video/video_track_observer.h"
 #include "epic_rtc_helper/memory/ref_count_impl_helper.h"
 
+#define UE_API PIXELSTREAMING2RTC_API
+
 class IPixelStreaming2VideoTrackObserver;
 
 namespace UE::PixelStreaming2
 {	
-	class PIXELSTREAMING2RTC_API FEpicRtcVideoTrackObserverFactory : public EpicRtcVideoTrackObserverFactoryInterface
+	class FEpicRtcVideoTrackObserverFactory : public EpicRtcVideoTrackObserverFactoryInterface
 	{
 	public:
-		FEpicRtcVideoTrackObserverFactory(TObserverVariant<IPixelStreaming2VideoTrackObserver> UserObserver);
+		UE_API FEpicRtcVideoTrackObserverFactory(TObserverVariant<IPixelStreaming2VideoTrackObserver> UserObserver);
 		virtual ~FEpicRtcVideoTrackObserverFactory() = default;
 
 	public:
 		// Begin EpicRtcVideoTrackObserverFactoryInterface
-		virtual EpicRtcErrorCode CreateVideoTrackObserver(const EpicRtcStringView ParticipantId, const EpicRtcStringView VideoTrackId, EpicRtcVideoTrackObserverInterface** OutVideoTrackObserver) override;
+		UE_API virtual EpicRtcErrorCode CreateVideoTrackObserver(const EpicRtcStringView ParticipantId, const EpicRtcStringView VideoTrackId, EpicRtcVideoTrackObserverInterface** OutVideoTrackObserver) override;
 		// End EpicRtcVideoTrackObserverFactoryInterface
 
 	public:
@@ -31,4 +33,6 @@ namespace UE::PixelStreaming2
 		TObserverVariant<IPixelStreaming2VideoTrackObserver> UserObserver;
 	};
 
-} // namespace UE::PixelStreaming2
\ No newline at end of file
+} // namespace UE::PixelStreaming2
+
+#undef UE_API
diff --git a/Engine/Plugins/Media/PixelStreaming2/Source/PixelStreaming2RTC/Internal/EpicRtcWebsocket.h b/Engine/Plugins/Media/PixelStreaming2/Source/PixelStreaming2RTC/Internal/EpicRtcWebsocket.h
--- a/Engine/Plugins/Media/PixelStreaming2/Source/PixelStreaming2RTC/Internal/EpicRtcWebsocket.h
+++ b/Engine/Plugins/Media/PixelStreaming2/Source/PixelStreaming2RTC/Internal/EpicRtcWebsocket.h
@@ -8,36 +8,38 @@
 #include "epic_rtc/plugins/signalling/websocket.h"
 #include "epic_rtc_helper/memory/ref_count_impl_helper.h"
 
+#define UE_API PIXELSTREAMING2RTC_API
+
 class IWebSocket;
 
 namespace UE::PixelStreaming2
 {
-	class PIXELSTREAMING2RTC_API FEpicRtcWebsocket : public EpicRtcWebsocketInterface, public FTickableGameObject
+	class FEpicRtcWebsocket : public EpicRtcWebsocketInterface, public FTickableGameObject
 	{
 	public:
-		FEpicRtcWebsocket(bool bKeepAlive = true, TSharedPtr<IWebSocket> WebSocket = nullptr);
+		UE_API FEpicRtcWebsocket(bool bKeepAlive = true, TSharedPtr<IWebSocket> WebSocket = nullptr);
 		virtual ~FEpicRtcWebsocket() = default;
 
 		// Begin EpicRtcWebsocketInterface
-		virtual EpicRtcBool Connect(EpicRtcStringView Url, EpicRtcWebsocketObserverInterface* Observer) override;
-		virtual void		Disconnect(const EpicRtcStringView Reason) override;
-		virtual void		Send(EpicRtcStringView Message) override;
+		UE_API virtual EpicRtcBool Connect(EpicRtcStringView Url, EpicRtcWebsocketObserverInterface* Observer) override;
+		UE_API virtual void		Disconnect(const EpicRtcStringView Reason) override;
+		UE_API virtual void		Send(EpicRtcStringView Message) override;
 		// End EpicRtcWebsocketInterface
 
 		// Begin FTickableGameObject
-		virtual void		Tick(float DeltaTime) override;
+		UE_API virtual void		Tick(float DeltaTime) override;
 		virtual bool		IsTickableInEditor() const override { return true; }
 		virtual bool		IsTickableWhenPaused() const override { return true; }
 		FORCEINLINE TStatId GetStatId() const { RETURN_QUICK_DECLARE_CYCLE_STAT(EpicRtcWebSocket, STATGROUP_Tickables); }
 		// End FTickableGameObject
 
 	private:
-		void OnConnected();
-		void OnConnectionError(const FString& Error);
-		void OnClosed(int32 StatusCode, const FString& Reason, bool bWasClean);
-		void OnMessage(const FString& Msg);
-		void OnBinaryMessage(const uint8* Data, int32 Length, bool bIsLastFragment);
-		void KeepAlive();
+		UE_API void OnConnected();
+		UE_API void OnConnectionError(const FString& Error);
+		UE_API void OnClosed(int32 StatusCode, const FString& Reason, bool bWasClean);
+		UE_API void OnMessage(const FString& Msg);
+		UE_API void OnBinaryMessage(const uint8* Data, int32 Length, bool bIsLastFragment);
+		UE_API void KeepAlive();
 
 	private:
 		FDelegateHandle OnConnectedHandle;
@@ -61,4 +63,6 @@ namespace UE::PixelStreaming2
 		// End EpicRtcRefCountInterface
 	};
 
-} // namespace UE::PixelStreaming2
\ No newline at end of file
+} // namespace UE::PixelStreaming2
+
+#undef UE_API
diff --git a/Engine/Plugins/Media/PixelStreaming2/Source/PixelStreaming2RTC/Internal/EpicRtcWebsocketFactory.h b/Engine/Plugins/Media/PixelStreaming2/Source/PixelStreaming2RTC/Internal/EpicRtcWebsocketFactory.h
--- a/Engine/Plugins/Media/PixelStreaming2/Source/PixelStreaming2RTC/Internal/EpicRtcWebsocketFactory.h
+++ b/Engine/Plugins/Media/PixelStreaming2/Source/PixelStreaming2RTC/Internal/EpicRtcWebsocketFactory.h
@@ -8,9 +8,11 @@
 #include "epic_rtc/plugins/signalling/websocket_factory.h"
 #include "epic_rtc_helper/memory/ref_count_impl_helper.h"
 
+#define UE_API PIXELSTREAMING2RTC_API
+
 namespace UE::PixelStreaming2
 {
-	class PIXELSTREAMING2RTC_API FEpicRtcWebsocketFactory : public EpicRtcWebsocketFactoryInterface
+	class FEpicRtcWebsocketFactory : public EpicRtcWebsocketFactoryInterface
 	{
 	public:
 		FEpicRtcWebsocketFactory(bool bInSendKeepAlive = true)
@@ -20,7 +22,7 @@ namespace UE::PixelStreaming2
 
 		virtual ~FEpicRtcWebsocketFactory() = default;
 
-		virtual EpicRtcErrorCode CreateWebsocket(EpicRtcWebsocketInterface** OutWebsocket) override;
+		UE_API virtual EpicRtcErrorCode CreateWebsocket(EpicRtcWebsocketInterface** OutWebsocket) override;
 
 	public:
 		// Begin EpicRtcRefCountInterface
@@ -32,3 +34,5 @@ namespace UE::PixelStreaming2
 	};
 
 } // namespace UE::PixelStreaming2
+
+#undef UE_API
diff --git a/Engine/Plugins/Media/PixelStreaming2/Source/PixelStreaming2RTC/Internal/RTCStatsCollector.h b/Engine/Plugins/Media/PixelStreaming2/Source/PixelStreaming2RTC/Internal/RTCStatsCollector.h
--- a/Engine/Plugins/Media/PixelStreaming2/Source/PixelStreaming2RTC/Internal/RTCStatsCollector.h
+++ b/Engine/Plugins/Media/PixelStreaming2/Source/PixelStreaming2RTC/Internal/RTCStatsCollector.h
@@ -11,6 +11,8 @@
 
 #include "epic_rtc/core/stats.h"
 
+#define UE_API PIXELSTREAMING2RTC_API
+
 namespace UE::PixelStreaming2
 {
 	class FStats;
@@ -36,7 +38,7 @@ namespace UE::PixelStreaming2
 		GRAPH = 1 << 1,
 	};
 
-	struct PIXELSTREAMING2RTC_API FStatConfig
+	struct FStatConfig
 	{
 		FName			 Name;
 		TOptional<FName> Alias;
@@ -45,25 +47,25 @@ namespace UE::PixelStreaming2
 
 	using FStatVariant = TVariant<TYPE_OF_NULLPTR, double, FString, bool>;
 
-	class PIXELSTREAMING2RTC_API FStat
+	class FStat
 	{
 	public:
 		FStat() = default;
 
-		FStat(FStatConfig Config, double InitialValue, int NDecimalPlacesToPrint = 0, bool bSmooth = false);
-		FStat(FStatConfig Config, FString InitialValue);
-		FStat(FStatConfig Config, bool bInitialValue);
-		FStat(const FStat& Other);
+		UE_API FStat(FStatConfig Config, double InitialValue, int NDecimalPlacesToPrint = 0, bool bSmooth = false);
+		UE_API FStat(FStatConfig Config, FString InitialValue);
+		UE_API FStat(FStatConfig Config, bool bInitialValue);
+		UE_API FStat(const FStat& Other);
 
-		bool IsNumeric() const;
+		UE_API bool IsNumeric() const;
 
-		bool IsTextual() const;
+		UE_API bool IsTextual() const;
 
-		bool IsBoolean() const;
+		UE_API bool IsBoolean() const;
 
-		FString ToString();
+		UE_API FString ToString();
 
-		bool SetValue(FStatVariant ValueVariant);
+		UE_API bool SetValue(FStatVariant ValueVariant);
 
 		template <typename T>
 		T GetValue()
@@ -95,21 +97,21 @@ namespace UE::PixelStreaming2
 		template <>
 		bool GetPrevValue();
 
-		bool operator==(const FStat& Other) const;
+		UE_API bool operator==(const FStat& Other) const;
 
-		bool IsHidden();
+		UE_API bool IsHidden();
 
-		bool ShouldGraph();
+		UE_API bool ShouldGraph();
 
-		bool ShouldDisplayText();
+		UE_API bool ShouldDisplayText();
 
-		FName GetName() const;
+		UE_API FName GetName() const;
 
-		FName GetDisplayName() const;
+		UE_API FName GetDisplayName() const;
 
 	private:
-		double CalcMA(double InPrevAvg, int InNumSamples, double InValue);
-		double CalcEMA(double InPrevAvg, int InNumSamples, double InValue);
+		UE_API double CalcMA(double InPrevAvg, int InNumSamples, double InValue);
+		UE_API double CalcEMA(double InPrevAvg, int InNumSamples, double InValue);
 
 	protected:
 		FName			 Name;
@@ -131,19 +133,19 @@ namespace UE::PixelStreaming2
 		return GetTypeHash(Obj.GetName());
 	}
 
-	class PIXELSTREAMING2RTC_API FRTCStatsCollector
+	class FRTCStatsCollector
 	{
 	public:
-		static TSharedPtr<FRTCStatsCollector> Create(const FString& PlayerId);
+		static UE_API TSharedPtr<FRTCStatsCollector> Create(const FString& PlayerId);
 
 	private:
-		FRTCStatsCollector();
-		explicit FRTCStatsCollector(const FString& PlayerId);
+		UE_API FRTCStatsCollector();
+		UE_API explicit FRTCStatsCollector(const FString& PlayerId);
 
-		void OnWebRTCDisableStatsChanged(IConsoleVariable* Var);
+		UE_API void OnWebRTCDisableStatsChanged(IConsoleVariable* Var);
 
 	public:
-		void Process(const EpicRtcConnectionStats& InStats);
+		UE_API void Process(const EpicRtcConnectionStats& InStats);
 
 	private:
 		class FStatsSink
@@ -267,4 +269,6 @@ namespace UE::PixelStreaming2
 
 		TUniquePtr<FCandidatePairStatsSink> CandidatePairStatsSink;
 	};
-} // namespace UE::PixelStreaming2
\ No newline at end of file
+} // namespace UE::PixelStreaming2
+
+#undef UE_API
diff --git a/Engine/Plugins/Media/PixelStreaming2/Source/PixelStreaming2RTC/Private/Blueprints/PixelStreaming2Peer.cpp b/Engine/Plugins/Media/PixelStreaming2/Source/PixelStreaming2RTC/Private/Blueprints/PixelStreaming2Peer.cpp
--- a/Engine/Plugins/Media/PixelStreaming2/Source/PixelStreaming2RTC/Private/Blueprints/PixelStreaming2Peer.cpp
+++ b/Engine/Plugins/Media/PixelStreaming2/Source/PixelStreaming2RTC/Private/Blueprints/PixelStreaming2Peer.cpp
@@ -10,6 +10,8 @@
 #include "SoundGenerator.h"
 #include "UtilsString.h"
 
+#include UE_INLINE_GENERATED_CPP_BY_NAME(PixelStreaming2Peer)
+
 uint32_t UPixelStreaming2Peer::PlayerId = 0;
 
 UPixelStreaming2Peer::UPixelStreaming2Peer(const FObjectInitializer& ObjectInitializer)
@@ -462,4 +464,4 @@ void UPixelStreaming2Peer::OnStatsReady(const FString& PeerId, const EpicRtcConn
 	}
 
 	StatsCollector->Process(ConnectionStats);
-}
\ No newline at end of file
+}
diff --git a/Engine/Plugins/Media/PixelStreaming2/Source/PixelStreaming2RTC/Private/EpicRtcAudioTrackObserver.cpp b/Engine/Plugins/Media/PixelStreaming2/Source/PixelStreaming2RTC/Private/EpicRtcAudioTrackObserver.cpp
--- a/Engine/Plugins/Media/PixelStreaming2/Source/PixelStreaming2RTC/Private/EpicRtcAudioTrackObserver.cpp
+++ b/Engine/Plugins/Media/PixelStreaming2/Source/PixelStreaming2RTC/Private/EpicRtcAudioTrackObserver.cpp
@@ -2,6 +2,8 @@
 
 #include "EpicRtcAudioTrackObserver.h"
 
+#include UE_INLINE_GENERATED_CPP_BY_NAME(EpicRtcAudioTrackObserver)
+
 namespace UE::PixelStreaming2
 {
 	FEpicRtcAudioTrackObserver::FEpicRtcAudioTrackObserver(TObserverVariant<IPixelStreaming2AudioTrackObserver> UserObserver)
@@ -41,4 +43,4 @@ namespace UE::PixelStreaming2
 		}
 	}
 
-} // namespace UE::PixelStreaming2
\ No newline at end of file
+} // namespace UE::PixelStreaming2
diff --git a/Engine/Plugins/Media/PixelStreaming2/Source/PixelStreaming2RTC/Private/EpicRtcDataTrackObserver.cpp b/Engine/Plugins/Media/PixelStreaming2/Source/PixelStreaming2RTC/Private/EpicRtcDataTrackObserver.cpp
--- a/Engine/Plugins/Media/PixelStreaming2/Source/PixelStreaming2RTC/Private/EpicRtcDataTrackObserver.cpp
+++ b/Engine/Plugins/Media/PixelStreaming2/Source/PixelStreaming2RTC/Private/EpicRtcDataTrackObserver.cpp
@@ -2,6 +2,8 @@
 
 #include "EpicRtcDataTrackObserver.h"
 
+#include UE_INLINE_GENERATED_CPP_BY_NAME(EpicRtcDataTrackObserver)
+
 namespace UE::PixelStreaming2
 {
 	FEpicRtcDataTrackObserver::FEpicRtcDataTrackObserver(TObserverVariant<IPixelStreaming2DataTrackObserver> UserObserver)
@@ -32,4 +34,4 @@ namespace UE::PixelStreaming2
 			UserObserver->OnDataTrackError(DataTrack, Error);
 		}
 	}
-} // namespace UE::PixelStreaming2
\ No newline at end of file
+} // namespace UE::PixelStreaming2
diff --git a/Engine/Plugins/Media/PixelStreaming2/Source/PixelStreaming2RTC/Private/EpicRtcRoomObserver.cpp b/Engine/Plugins/Media/PixelStreaming2/Source/PixelStreaming2RTC/Private/EpicRtcRoomObserver.cpp
--- a/Engine/Plugins/Media/PixelStreaming2/Source/PixelStreaming2RTC/Private/EpicRtcRoomObserver.cpp
+++ b/Engine/Plugins/Media/PixelStreaming2/Source/PixelStreaming2RTC/Private/EpicRtcRoomObserver.cpp
@@ -2,6 +2,8 @@
 
 #include "EpicRtcRoomObserver.h"
 
+#include UE_INLINE_GENERATED_CPP_BY_NAME(EpicRtcRoomObserver)
+
 namespace UE::PixelStreaming2
 {
 	FEpicRtcRoomObserver::FEpicRtcRoomObserver(TObserverVariant<IPixelStreaming2RoomObserver> UserObserver)
@@ -83,4 +85,4 @@ namespace UE::PixelStreaming2
 		}
 	}
 
-} // namespace UE::PixelStreaming2
\ No newline at end of file
+} // namespace UE::PixelStreaming2
diff --git a/Engine/Plugins/Media/PixelStreaming2/Source/PixelStreaming2RTC/Private/EpicRtcSessionObserver.cpp b/Engine/Plugins/Media/PixelStreaming2/Source/PixelStreaming2RTC/Private/EpicRtcSessionObserver.cpp
--- a/Engine/Plugins/Media/PixelStreaming2/Source/PixelStreaming2RTC/Private/EpicRtcSessionObserver.cpp
+++ b/Engine/Plugins/Media/PixelStreaming2/Source/PixelStreaming2RTC/Private/EpicRtcSessionObserver.cpp
@@ -2,6 +2,8 @@
 
 #include "EpicRtcSessionObserver.h"
 
+#include UE_INLINE_GENERATED_CPP_BY_NAME(EpicRtcSessionObserver)
+
 namespace UE::PixelStreaming2
 {
 
diff --git a/Engine/Plugins/Media/PixelStreaming2/Source/PixelStreaming2RTC/Private/EpicRtcVideoTrackObserver.cpp b/Engine/Plugins/Media/PixelStreaming2/Source/PixelStreaming2RTC/Private/EpicRtcVideoTrackObserver.cpp
--- a/Engine/Plugins/Media/PixelStreaming2/Source/PixelStreaming2RTC/Private/EpicRtcVideoTrackObserver.cpp
+++ b/Engine/Plugins/Media/PixelStreaming2/Source/PixelStreaming2RTC/Private/EpicRtcVideoTrackObserver.cpp
@@ -2,6 +2,8 @@
 
 #include "EpicRtcVideoTrackObserver.h"
 
+#include UE_INLINE_GENERATED_CPP_BY_NAME(EpicRtcVideoTrackObserver)
+
 namespace UE::PixelStreaming2
 {
 
@@ -59,4 +61,4 @@ namespace UE::PixelStreaming2
 		return false;
 	}
 
-} // namespace UE::PixelStreaming2
\ No newline at end of file
+} // namespace UE::PixelStreaming2
diff --git a/Engine/Plugins/Media/PixelStreaming2/Source/PixelStreaming2RTC/Private/Tests/DelegateTest.cpp b/Engine/Plugins/Media/PixelStreaming2/Source/PixelStreaming2RTC/Private/Tests/DelegateTest.cpp
--- a/Engine/Plugins/Media/PixelStreaming2/Source/PixelStreaming2RTC/Private/Tests/DelegateTest.cpp
+++ b/Engine/Plugins/Media/PixelStreaming2/Source/PixelStreaming2RTC/Private/Tests/DelegateTest.cpp
@@ -2,6 +2,8 @@
 
 #include "DelegateTest.h"
 
+#include UE_INLINE_GENERATED_CPP_BY_NAME(DelegateTest)
+
 #if WITH_TESTS
 
 #include "CodecUtils.h"
diff --git a/Engine/Plugins/Media/PixelStreaming2/Source/PixelStreaming2RTC/Public/IPixelStreaming2RTCModule.h b/Engine/Plugins/Media/PixelStreaming2/Source/PixelStreaming2RTC/Public/IPixelStreaming2RTCModule.h
--- a/Engine/Plugins/Media/PixelStreaming2/Source/PixelStreaming2RTC/Public/IPixelStreaming2RTCModule.h
+++ b/Engine/Plugins/Media/PixelStreaming2/Source/PixelStreaming2RTC/Public/IPixelStreaming2RTCModule.h
@@ -13,7 +13,7 @@
  * This class provides access to streamers, video and audio producers, and handles starting and stopping the streaming process.
  * It also allows interaction with the signalling server and manages the lifecycle of streamers within the Pixel Streaming system.
  */
-class PIXELSTREAMING2RTC_API IPixelStreaming2RTCModule : public IModuleInterface
+class IPixelStreaming2RTCModule : public IModuleInterface
 {
 public:
 	/**
diff --git a/Engine/Plugins/Media/PixelStreaming2/Source/PixelStreaming2RTC/Public/IPixelStreaming2Stats.h b/Engine/Plugins/Media/PixelStreaming2/Source/PixelStreaming2RTC/Public/IPixelStreaming2Stats.h
--- a/Engine/Plugins/Media/PixelStreaming2/Source/PixelStreaming2RTC/Public/IPixelStreaming2Stats.h
+++ b/Engine/Plugins/Media/PixelStreaming2/Source/PixelStreaming2RTC/Public/IPixelStreaming2Stats.h
@@ -4,18 +4,20 @@
 
 #include "UObject/NameTypes.h"
 
+#define UE_API PIXELSTREAMING2RTC_API
+
 /**
  * Add stats that Pixel Streaming can display on screen or be used internally.
  */
 
-class PIXELSTREAMING2RTC_API IPixelStreaming2Stats
+class IPixelStreaming2Stats
 {
 public:
 	/**
 	 * Singleton-like access to the Pixel Streaming stats public API.
 	 * @return Returns singleton instance of Pixel Streaming stats.
 	 */
-	static IPixelStreaming2Stats& Get();
+	static UE_API IPixelStreaming2Stats& Get();
 
 	/**
 	 * Graphs a value in a real-time plot that is visible when `stat PixelStreaming2Graphs` is used.
@@ -27,4 +29,6 @@ public:
 	 * @param InRefValue A reference value that we draw in the graph as a horizontal line, this can be useful to see if we are within some threshold.
 	 */
 	virtual void GraphValue(FName InName, float Value, int InSamples, float InMinRange, float InMaxRange, float InRefValue = 0.0f) = 0;
-};
\ No newline at end of file
+};
+
+#undef UE_API
diff --git a/Engine/Plugins/Media/PixelStreaming2/Source/PixelStreaming2Servers/Public/PixelStreaming2Servers.h b/Engine/Plugins/Media/PixelStreaming2/Source/PixelStreaming2Servers/Public/PixelStreaming2Servers.h
--- a/Engine/Plugins/Media/PixelStreaming2/Source/PixelStreaming2Servers/Public/PixelStreaming2Servers.h
+++ b/Engine/Plugins/Media/PixelStreaming2/Source/PixelStreaming2Servers/Public/PixelStreaming2Servers.h
@@ -19,7 +19,7 @@ namespace UE::PixelStreaming2Servers
 	/**
 	 * Configuration to control behaviour when launching any of the Pixel Streaming servers.
 	 **/
-	struct PIXELSTREAMING2SERVERS_API FLaunchArgs
+	struct FLaunchArgs
 	{
 		// Arguments passed to the actual server when its process is started.
 		FString ProcessArgs = TEXT("");
@@ -60,7 +60,7 @@ namespace UE::PixelStreaming2Servers
 	/**
 	 * Interface for all Pixel Streaming servers.
 	 **/
-	class PIXELSTREAMING2SERVERS_API IServer
+	class IServer
 	{
 	public:
 		virtual ~IServer() = default;
diff --git a/Engine/Plugins/Media/PixelStreaming2/Source/PixelStreaming2Settings/Internal/PixelStreaming2PluginSettings.h b/Engine/Plugins/Media/PixelStreaming2/Source/PixelStreaming2Settings/Internal/PixelStreaming2PluginSettings.h
--- a/Engine/Plugins/Media/PixelStreaming2/Source/PixelStreaming2Settings/Internal/PixelStreaming2PluginSettings.h
+++ b/Engine/Plugins/Media/PixelStreaming2/Source/PixelStreaming2Settings/Internal/PixelStreaming2PluginSettings.h
@@ -16,6 +16,8 @@
 
 #include "PixelStreaming2PluginSettings.generated.h"
 
+#define UE_API PIXELSTREAMING2SETTINGS_API
+
 namespace UE::PixelStreaming2
 {
 	// This variable exists to back CVarEncoderCodec and should only be used by the FEpicRtcVideo(Encoder/Decoder)Initializer. When doing software fallback,
@@ -77,93 +79,93 @@ enum class EInputControllerMode : uint8
 
 // Config loaded/saved to an .ini file.
 // It is also exposed through the plugin settings page in editor.
-UCLASS(config = Game, defaultconfig, meta = (DisplayName = "PixelStreaming2"))
-class PIXELSTREAMING2SETTINGS_API UPixelStreaming2PluginSettings : public UDeveloperSettings
+UCLASS(MinimalAPI, config = Game, defaultconfig, meta = (DisplayName = "PixelStreaming2"))
+class UPixelStreaming2PluginSettings : public UDeveloperSettings
 {
 	GENERATED_BODY()
 
-	virtual ~UPixelStreaming2PluginSettings();
+	UE_API virtual ~UPixelStreaming2PluginSettings();
 
 public:
 	// clang-format off
-	static TAutoConsoleVariable<bool> CVarLogStats;
+	static UE_API TAutoConsoleVariable<bool> CVarLogStats;
 	UPROPERTY(config, EditAnywhere, Category = "PixelStreaming", meta = (
 		DisplayName = "Log Pixel Streaming Stats",
 		ToolTip = "Whether to show PixelStreaming stats in the log (default: false)."
 		))
 	bool LogStats = false;
 
-	static TAutoConsoleVariable<FString> CVarEpicRtcLogFilter;
+	static UE_API TAutoConsoleVariable<FString> CVarEpicRtcLogFilter;
 	UPROPERTY(config, EditAnywhere, Category = "PixelStreaming", meta = (
 		DisplayName = "Log messages from EpicRtc to filter",
 		ToolTip = "Double forward slash (\"//\") separated list of regex patterns to filter from the EpicRtc logs (default: \"\")."
 		))
 	FString EpicRtcLogFilter = "";
 
-	static TAutoConsoleVariable<bool> CVarDisableLatencyTester;
+	static UE_API TAutoConsoleVariable<bool> CVarDisableLatencyTester;
 	UPROPERTY(config, EditAnywhere, Category = "PixelStreaming", meta = (
 		DisplayName = "Disable Latency Tester",
 		ToolTip = "If true disables latency tester being triggerable."
 		))
 	bool DisableLatencyTester = false;
 
-	static TAutoConsoleVariable<FString> CVarInputController;
+	static UE_API TAutoConsoleVariable<FString> CVarInputController;
 	UPROPERTY(config, EditAnywhere, Category = "PixelStreaming", meta = (
 		DisplayName = "Input Controller Mode",
 		ToolTip = "If true disables latency tester being triggerable."
 		))
 	EInputControllerMode InputController = EInputControllerMode::Any;
 
-	static TAutoConsoleVariable<bool> CVarDecoupleFramerate;
+	static UE_API TAutoConsoleVariable<bool> CVarDecoupleFramerate;
 	UPROPERTY(config, EditAnywhere, Category = "PixelStreaming", meta = (
 		DisplayName = "Decouple Frame Rate",
 		ToolTip = "Whether we should only stream as fast as we render or at some fixed interval. Coupled means only stream what we render."
 		))
 	bool DecoupleFramerate = false;
 
-	static TAutoConsoleVariable<float> CVarDecoupleWaitFactor;
+	static UE_API TAutoConsoleVariable<float> CVarDecoupleWaitFactor;
 	UPROPERTY(config, EditAnywhere, Category = "PixelStreaming", meta = (
 		DisplayName = "Decouple Wait Factor",
 		ToolTip = "Frame rate factor to wait for a captured frame when streaming in decoupled mode. Higher factor waits longer but may also result in higher latency."
 		))
 	float DecoupleWaitFactor = 1.0f;
 
-	static TAutoConsoleVariable<float> CVarSignalingReconnectInterval;
+	static UE_API TAutoConsoleVariable<float> CVarSignalingReconnectInterval;
 	UPROPERTY(config, EditAnywhere, Category = "PixelStreaming", meta = (
 		DisplayName = "Signaling Reconnect Interval",
 		ToolTip = "Changes the number of seconds between attempted reconnects to the signaling server. This is useful for reducing the log spam produced from attempted reconnects. A value <= 0 results in no reconnect. Default: 2.0s"
 		))
 	float SignalingReconnectInterval = 2.0f;
 
-	static TAutoConsoleVariable<float> CVarSignalingMaxReconnectAttempts;
+	static UE_API TAutoConsoleVariable<float> CVarSignalingMaxReconnectAttempts;
 	UPROPERTY(config, EditAnywhere, Category = "PixelStreaming", meta = (
 		DisplayName = "Signaling Max Reconnect Attempts",
 		ToolTip = "Changes the number of attempts that will be made to reconnect to the signalling server. This is useful for triggering application shutdowns if this value is exceeded. A value of < 0 results in unlimited attempts. Default: -1"
 		))
 	int32 SignalingMaxReconnectAttempts = -1;
 
-	static TAutoConsoleVariable<float> CVarSignalingKeepAliveInterval;
+	static UE_API TAutoConsoleVariable<float> CVarSignalingKeepAliveInterval;
 	UPROPERTY(config, EditAnywhere, Category = "PixelStreaming", meta = (
 		DisplayName = "Signaling Reconnect Interval",
 		ToolTip = "Frame rate factor to wait for a captured frame when streaming in decoupled mode. Higher factor waits longer but may also result in higher latency."
 		))
 	float SignalingKeepAliveInterval = 30.0f;
 
-	static TAutoConsoleVariable<bool> CVarUseMediaCapture;
+	static UE_API TAutoConsoleVariable<bool> CVarUseMediaCapture;
 	UPROPERTY(config, EditAnywhere, Category = "PixelStreaming", meta = (
 		DisplayName = "Use Media Capture",
 		ToolTip = "Use Media Capture from MediaIOFramework to capture frames rather than Pixel Streamings internal backbuffer sources."
 		))
 	bool UseMediaCapture = true;
 
-	static TAutoConsoleVariable<FString> CVarDefaultStreamerID;
+	static UE_API TAutoConsoleVariable<FString> CVarDefaultStreamerID;
 	UPROPERTY(config, EditAnywhere, Category = "PixelStreaming", meta = (
 		DisplayName = "Default Streamer ID",
 		ToolTip = "Default Streamer ID to be used when not specified elsewhere."
 		))
 	FString DefaultStreamerID = TEXT("DefaultStreamer");
 
-	static TAutoConsoleVariable<FString> CVarDefaultStreamerType;
+	static UE_API TAutoConsoleVariable<FString> CVarDefaultStreamerType;
 	UPROPERTY(config, EditAnywhere, Category = "PixelStreaming", meta = (
 		DisplayName = "Default Streamer Type",
 		ToolTip = "Default Streamer Type to be used when not specified elsewhere. This value should match a type registered by an IPixelStreaming2StreamerFactory",
@@ -171,30 +173,30 @@ public:
 		))
 	FString DefaultStreamerType = TEXT("DefaultRtc");
 
-	static TAutoConsoleVariable<bool> CVarAutoStartStream;
+	static UE_API TAutoConsoleVariable<bool> CVarAutoStartStream;
 	UPROPERTY(config, EditAnywhere, Category = "PixelStreaming", meta = (
 		DisplayName = "Automatically Start Streaming",
 		ToolTip = "Configure the PixelStreaming2 plugin to automatically start streaming once loaded (if not in editor). You may wish to set this value to false and manually call StartStreaming at a later point from your c++ code. Default: true"
 		))
 	bool AutoStartStream = true;
 
-	static TAutoConsoleVariable<FString> CVarConnectionURL;
+	static UE_API TAutoConsoleVariable<FString> CVarConnectionURL;
 	UPROPERTY(config, EditAnywhere, Category = "PixelStreaming", meta = (
 		DisplayName = "Default Connection URL",
 		ToolTip = "Default URL to connect to. This can be a URL to a signalling server or some other endpoint with the format (protocol)://(host):(port)"
 		))
 	FString ConnectionURL;
 
-	static FAutoConsoleVariableDeprecated CVarSignallingURL;
+	static UE_API FAutoConsoleVariableDeprecated CVarSignallingURL;
 
-	static TAutoConsoleVariable<bool> CVarCaptureUseFence;
+	static UE_API TAutoConsoleVariable<bool> CVarCaptureUseFence;
 	UPROPERTY(config, EditAnywhere, Category = "PixelStreaming", meta = (
 		DisplayName = "Capture Using Fence",
 		ToolTip = "Whether the texture copy we do during image capture should use a fence or not (non-fenced is faster but less safe)."
 		))
 	bool CaptureUseFence = true;
 
-	static TAutoConsoleVariable<bool> CVarDebugDumpAudio;
+	static UE_API TAutoConsoleVariable<bool> CVarDebugDumpAudio;
 
 	// Begin Cursor Settings
 	/**
@@ -217,21 +219,21 @@ public:
 	FSoftClassPath HiddenCursorClassName;
 
 	// Begin Encoder Settings
-	static TAutoConsoleVariable<int32> CVarEncoderTargetBitrate;
+	static UE_API TAutoConsoleVariable<int32> CVarEncoderTargetBitrate;
 	UPROPERTY(config, EditAnywhere, Category = "EncoderSettings", meta = (
 		DisplayName = "Target Bitrate",
 		ToolTip = "Target bitrate (bps). Ignore the bitrate WebRTC wants (not recommended). Set to -1 to disable. Default -1.",
 		ClampMin = -1))
 	int32 EncoderTargetBitrate = -1;
 
-	static TAutoConsoleVariable<int32>	 CVarEncoderMinQuality;
+	static UE_API TAutoConsoleVariable<int32>	 CVarEncoderMinQuality;
 	UPROPERTY(config, EditAnywhere, Category = "EncoderSettings", meta = (
 		DisplayName = "Encoder Minimum Quality",
 		ToolTip = "0-100, Higher values result in a better minimum quality but higher average bitrates. Default 0 - i.e. no limit on a minimum Quality.",
 		ClampMin = 0, ClampMax = 100))
 	int32 EncoderMinQuality = 0;
 
-	static TAutoConsoleVariable<int32>	CVarEncoderMaxQuality;
+	static UE_API TAutoConsoleVariable<int32>	CVarEncoderMaxQuality;
 	UPROPERTY(config, EditAnywhere, Category = "EncoderSettings", meta = (
 		DisplayName = "Encoder Maximum Quality",
 		ToolTip = "0-100, Lower values result in lower average bitrates but reduces maximum quality. Default 100 - i.e. no limit on a maximum QP.",
@@ -239,7 +241,7 @@ public:
 	int32 EncoderMaxQuality = 100;
 
 	// not directly connected to cvar due to string to enum conversion
-	static TAutoConsoleVariable<FString> CVarEncoderQualityPreset;
+	static UE_API TAutoConsoleVariable<FString> CVarEncoderQualityPreset;
 	UPROPERTY(config, EditAnywhere, Category = "EncoderSettings", meta = (
 		DisplayName = "Default Encoding Quality Profile",
 		ToolTip = "PixelStreaming Encoder presets that affecting Quality vs Bitrate. Supported modes are `ULTRA_LOW_QUALITY`, `LOW_QUALITY`, `DEFAULT`, `HIGH_QUALITY`, `LOSSLESS`"
@@ -247,28 +249,28 @@ public:
 	EAVPreset QualityPreset = EAVPreset::Default;
 
 	// not directly connected to cvar due to string to enum conversion
-	static TAutoConsoleVariable<FString> CVarEncoderLatencyMode;
+	static UE_API TAutoConsoleVariable<FString> CVarEncoderLatencyMode;
 	UPROPERTY(config, EditAnywhere, Category = "EncoderSettings", meta = (
 		DisplayName = "Default Encoding Quality Profile",
 		ToolTip = "PixelStreaming Encoder presets that affecting Quality vs Latency. Supported modes are `ULTRA_LOW_LATENCY`, `LOW_LATENCY`, `DEFAULT`"
 		))
 	EAVLatencyMode LatencyMode = EAVLatencyMode::UltraLowLatency;
 
-	static TAutoConsoleVariable<int32> CVarEncoderKeyframeInterval;
+	static UE_API TAutoConsoleVariable<int32> CVarEncoderKeyframeInterval;
 	UPROPERTY(config, EditAnywhere, Category = "EncoderSettings", meta = (
 		DisplayName = "Default Keyframe Interval",
 		ToolTip = "How many frames before a key frame is sent. Default: -1 which disables the sending of periodic key frames. Note: NVENC reqires a reinitialization when this changes."
 		))
 	int32 KeyframeInterval = -1;
 
-	static TAutoConsoleVariable<int32> CVarEncoderMaxSessions;
+	static UE_API TAutoConsoleVariable<int32> CVarEncoderMaxSessions;
 	UPROPERTY(config, EditAnywhere, Category = "EncoderSettings", meta = (
 		DisplayName = "Default Max Number of Encoding Session",
 		ToolTip = "-1 implies no limit. Maximum number of concurrent hardware encoder sessions for Pixel Streaming. Note: Geforce gpus only support 8 concurrent sessions and will rollover to software encoding when that number is exceeded."
 		))
 	int32 MaxSessions = -1;
 
-	static TAutoConsoleVariable<bool> CVarEncoderEnableSimulcast;
+	static UE_API TAutoConsoleVariable<bool> CVarEncoderEnableSimulcast;
 	UPROPERTY(config, EditAnywhere, Category = "EncoderSettings", meta = (
 		DisplayName = "Enable Simulcast",
 		ToolTip = "Enables simulcast. When enabled, the encoder will encode at full resolution, 1/2 resolution and 1/4 resolution simultaneously. Note: Simulcast is only supported with `H264` and `VP8` and you must use the SFU from the infrastructure to fully utilise this functionality.",
@@ -277,7 +279,7 @@ public:
 		))
 	bool EnableSimulcast = false;
 
-	static TAutoConsoleVariable<FString> CVarEncoderCodec;
+	static UE_API TAutoConsoleVariable<FString> CVarEncoderCodec;
 	UPROPERTY(config, EditAnywhere, Category = "EncoderSettings", meta = (
 		DisplayName = "Preferred Encoder Codec",
 		ToolTip = "Preferred encoder codec signalled during connection establishment.",
@@ -287,7 +289,7 @@ public:
 		))
 	FString Codec = UE::PixelStreaming2::GetCVarStringFromEnum(EVideoCodec::H264);
 
-	static TAutoConsoleVariable<FString> CVarEncoderScalabilityMode;
+	static UE_API TAutoConsoleVariable<FString> CVarEncoderScalabilityMode;
 	UPROPERTY(config, EditAnywhere, Category = "EncoderSettings", meta = (
 		DisplayName = "Default Scalability Mode",
 		ToolTip = "Indicates number of spatial and temporal layers used, default: L1T1. For a full list of values refer to https://www.w3.org/TR/webrtc-svc/#scalabilitymodes*",
@@ -297,7 +299,7 @@ public:
 		))
 	FString ScalabilityMode = UE::PixelStreaming2::GetCVarStringFromEnum(EScalabilityMode::L1T1);
 
-	static TAutoConsoleVariable<FString> CVarEncoderH264Profile;
+	static UE_API TAutoConsoleVariable<FString> CVarEncoderH264Profile;
 	UPROPERTY(config, EditAnywhere, Category = "EncoderSettings|H264", meta = (
 		DisplayName = "Default H264 profile",
 		ToolTip = "PixelStreaming Encoder profile. Supported modes are `AUTO`, `BASELINE`, `MAIN`, `HIGH`, `HIGH444`, `PROGRESSIVE_HIGH` or `CONSTRAINED_HIGH`",
@@ -305,102 +307,102 @@ public:
 		))
 	EH264Profile H264Profile = EH264Profile::Baseline;
 
-	static TAutoConsoleVariable<bool> CVarEncoderDebugDumpFrame;
+	static UE_API TAutoConsoleVariable<bool> CVarEncoderDebugDumpFrame;
 	// End Encoder Settings
 
 	// Begin WebRTC CVars
-	static TAutoConsoleVariable<int32> CVarWebRTCFps;
+	static UE_API TAutoConsoleVariable<int32> CVarWebRTCFps;
 	UPROPERTY(config, EditAnywhere, Category = "WebRtcSettings", meta = (
 		DisplayName = "Default WebRtc FPS",
 		ToolTip = "Framerate for WebRTC encoding. Default: 60"
 		))
 	int32 WebRTCFps = 60;
 
-	static TAutoConsoleVariable<int32> CVarWebRTCStartBitrate;
+	static UE_API TAutoConsoleVariable<int32> CVarWebRTCStartBitrate;
 	UPROPERTY(config, EditAnywhere, Category = "WebRtcSettings", meta = (
 		DisplayName = "Default WebRtc Start Bitrate",
 		ToolTip = "FStart bitrate (bps) that WebRTC will try begin the stream with. Must be between Min/Max bitrates. Default: 1000000"
 		))
 	int32 WebRTCStartBitrate = 1000000;
 
-	static TAutoConsoleVariable<int32> CVarWebRTCMinBitrate;
+	static UE_API TAutoConsoleVariable<int32> CVarWebRTCMinBitrate;
 	UPROPERTY(config, EditAnywhere, Category = "WebRtcSettings", meta = (
 		DisplayName = "Default WebRtc Minimum Bitrate",
 		ToolTip = "Min bitrate (bps) that WebRTC will not request below. Careful not to set too high otherwise WebRTC will just drop frames. Default: 100000"
 		))
 	int32 WebRTCMinBitrate = 100000;
 
-	static TAutoConsoleVariable<int32> CVarWebRTCMaxBitrate;
+	static UE_API TAutoConsoleVariable<int32> CVarWebRTCMaxBitrate;
 	UPROPERTY(config, EditAnywhere, Category = "WebRtcSettings", meta = (
 		DisplayName = "Default WebRtc Maximum Bitrate",
 		ToolTip = "Max bitrate (bps) that WebRTC will not request above. Default: 40000000 aka 40 megabits/per second."
 		))
 	int32 WebRTCMaxBitrate = 40000000;
 
-	static TAutoConsoleVariable<bool> CVarWebRTCDisableReceiveAudio;
+	static UE_API TAutoConsoleVariable<bool> CVarWebRTCDisableReceiveAudio;
 	UPROPERTY(config, EditAnywhere, Category = "WebRtcSettings", meta = (
 		DisplayName = "Disable receiving audio by default",
 		ToolTip = "Disables receiving audio from the browser into UE."
 		))
 	bool WebRTCDisableReceiveAudio = false;
 
-	static TAutoConsoleVariable<bool> CVarWebRTCDisableReceiveVideo;
+	static UE_API TAutoConsoleVariable<bool> CVarWebRTCDisableReceiveVideo;
 	UPROPERTY(config, EditAnywhere, Category = "WebRtcSettings", meta = (
 		DisplayName = "Disable receiving video by default",
 		ToolTip = "Disables receiving video from the browser into UE."
 		))
 	bool WebRTCDisableReceiveVideo = false;
 
-	static TAutoConsoleVariable<bool> CVarWebRTCDisableTransmitAudio;
+	static UE_API TAutoConsoleVariable<bool> CVarWebRTCDisableTransmitAudio;
 	UPROPERTY(config, EditAnywhere, Category = "WebRtcSettings", meta = (
 		DisplayName = "Disable transmiting audio by default",
 		ToolTip = "Disables transmiting audio to the browser."
 		))
 	bool WebRTCDisableTransmitAudio = false;
 
-	static TAutoConsoleVariable<bool> CVarWebRTCDisableTransmitVideo;
+	static UE_API TAutoConsoleVariable<bool> CVarWebRTCDisableTransmitVideo;
 	UPROPERTY(config, EditAnywhere, Category = "WebRtcSettings", meta = (
 		DisplayName = "Disable transmiting video by default",
 		ToolTip = "Disables transmiting video to the browser."
 		))
 	bool WebRTCDisableTransmitVideo = false;
 
-	static TAutoConsoleVariable<bool> CVarWebRTCDisableAudioSync;
+	static UE_API TAutoConsoleVariable<bool> CVarWebRTCDisableAudioSync;
 	UPROPERTY(config, EditAnywhere, Category = "WebRtcSettings", meta = (
 		DisplayName = "Disable Audio/Video track sychronisation",
 		ToolTip = "Disables the synchronization of audio and video tracks in WebRTC. This can be useful in low latency usecases where synchronization is not required."
 		))
 	bool WebRTCDisableAudioSync = true;
 
-	static TAutoConsoleVariable<bool> CVarWebRTCEnableFlexFec;
+	static UE_API TAutoConsoleVariable<bool> CVarWebRTCEnableFlexFec;
 	UPROPERTY(config, EditAnywhere, Category = "WebRtcSettings", meta = (
 		DisplayName = "Enable Flexible Forward Error Correction",
 		ToolTip = "Signals support for Flexible Forward Error Correction to WebRTC. This can cause a reduction in quality if total bitrate is low."
 		))
 	bool WebRTCEnableFlexFec = false;
 
-	static TAutoConsoleVariable<bool> CVarWebRTCDisableStats;
+	static UE_API TAutoConsoleVariable<bool> CVarWebRTCDisableStats;
 	UPROPERTY(config, EditAnywhere, Category = "WebRtcSettings", meta = (
 		DisplayName = "Disable WebRtc Stats",
 		ToolTip = "Disables the collection of WebRTC stats."
 		))
 	bool WebRTCDisableStats = false;
 
-	static TAutoConsoleVariable<float> CVarWebRTCStatsInterval;
+	static UE_API TAutoConsoleVariable<float> CVarWebRTCStatsInterval;
 	UPROPERTY(config, EditAnywhere, Category = "WebRtcSettings", meta = (
 		DisplayName = "Stats Interval",
 		ToolTip = "Configures how often WebRTC stats are collected."
 		))
 	float WebRTCStatsInterval = 1.f;
 
-	static TAutoConsoleVariable<bool> CVarWebRTCNegotiateCodecs;
+	static UE_API TAutoConsoleVariable<bool> CVarWebRTCNegotiateCodecs;
 	UPROPERTY(config, EditAnywhere, Category = "WebRtcSettings", meta = (
 		DisplayName = "WebRtc Negotiate Codecs",
 		ToolTip = "Whether PixelStreaming should send all its codecs during sdp handshake so peers can negotiate or just send a single selected codec."
 		))
 	bool WebRTCNegotiateCodecs = false;
 
-	static TAutoConsoleVariable<FString> CVarWebRTCCodecPreferences;
+	static UE_API TAutoConsoleVariable<FString> CVarWebRTCCodecPreferences;
 	UPROPERTY(config, EditAnywhere, Category = "WebRtcSettings", meta = (
 		DisplayName = "WebRtc Codec Preferences",
 		ToolTip = "The preference order PixelStreaming will specify during sdp handshake",
@@ -416,16 +418,16 @@ public:
 		UE::PixelStreaming2::GetCVarStringFromEnum(EVideoCodec::VP8)
 	 };
 
-	static TArray<EVideoCodec> GetCodecPreferences();
+	static UE_API TArray<EVideoCodec> GetCodecPreferences();
 
-	static TAutoConsoleVariable<float> CVarWebRTCAudioGain;
+	static UE_API TAutoConsoleVariable<float> CVarWebRTCAudioGain;
 	UPROPERTY(config, EditAnywhere, Category = "WebRtcSettings", meta = (
 		DisplayName = "Default Audio Gain",
 		ToolTip = "Sets the amount of gain to apply to audio. Default: 1.0"
 		))
 	float WebRTCAudioGain = 1.0f;
 
-	static TAutoConsoleVariable<FString> CVarWebRTCPortAllocatorFlags;
+	static UE_API TAutoConsoleVariable<FString> CVarWebRTCPortAllocatorFlags;
 	UPROPERTY(config, EditAnywhere, Category = "WebRtcSettings", meta = (
 		Bitmask,
 		BitmaskEnum = "/Script/PixelStreaming2Settings.EPortAllocatorFlags",
@@ -434,44 +436,44 @@ public:
 		))
 	int32 WebRTCPortAllocatorFlags = 0;
 
-	static EPortAllocatorFlags GetPortAllocationFlags();
+	static UE_API EPortAllocatorFlags GetPortAllocationFlags();
 
-	static TAutoConsoleVariable<int> CVarWebRTCMinPort;
+	static UE_API TAutoConsoleVariable<int> CVarWebRTCMinPort;
 	UPROPERTY(config, EditAnywhere, Category = "WebRtcSettings", meta = (
 		DisplayName = "Default WebRtc Min Port",
 		ToolTip = "Sets the minimum usable port for the WebRTC port allocator. Default: 49152"
 		))
 	int WebRTCMinPort = 49152;
 
-	static TAutoConsoleVariable<int> CVarWebRTCMaxPort;
+	static UE_API TAutoConsoleVariable<int> CVarWebRTCMaxPort;
 	UPROPERTY(config, EditAnywhere, Category = "WebRtcSettings", meta = (
 		DisplayName = "Default WebRtc Max Port",
 		ToolTip = "Sets the maximum usable port for the WebRTC port allocator. Default: 65535"
 		))
 	int WebRTCMaxPort = 65535;
 
-	static TAutoConsoleVariable<FString> CVarWebRTCFieldTrials;
+	static UE_API TAutoConsoleVariable<FString> CVarWebRTCFieldTrials;
 	UPROPERTY(config, EditAnywhere, Category = "WebRtcSettings|Field Trials", meta = (
 		DisplayName = "Manually defined WebRtc Field Trials",
 		ToolTip = "Sets the WebRTC field trials string. Format:\"TRIAL1/VALUE1/TRIAL2/VALUE2/\" see https://webrtc.googlesource.com/src/+/HEAD/g3doc/field-trials.md"
 		))
 	FString WebRTCFieldTrials;
 
-	static TAutoConsoleVariable<bool> CVarWebRTCDisableFrameDropper;
+	static UE_API TAutoConsoleVariable<bool> CVarWebRTCDisableFrameDropper;
 	UPROPERTY(config, EditAnywhere, Category = "WebRtcSettings|Field Trials", meta = (
 		DisplayName = "Disable Frame Dropper",
 		ToolTip = "Disables the WebRTC internal frame dropper using the field trial WebRTC-FrameDropper/Disabled/"
 		))
 	bool WebRTCDisableFrameDropper = false;
 
-	static TAutoConsoleVariable<float> CVarWebRTCVideoPacingMaxDelay;
+	static UE_API TAutoConsoleVariable<float> CVarWebRTCVideoPacingMaxDelay;
 	UPROPERTY(config, EditAnywhere, Category = "WebRtcSettings|Field Trials", meta = (
 		DisplayName = "Video Pacing Max Delay",
 		ToolTip = "Enables the WebRTC-Video-Pacing field trial and sets the max delay (ms) parameter. Default: -1.0f (values below zero are discarded.)"
 		))
 	float WebRTCVideoPacingMaxDelay = -1.0f;
 
-	static TAutoConsoleVariable<float> CVarWebRTCVideoPacingFactor;
+	static UE_API TAutoConsoleVariable<float> CVarWebRTCVideoPacingFactor;
 	UPROPERTY(config, EditAnywhere, Category = "WebRtcSettings|Field Trials", meta = (
 		DisplayName = "Video Pacing Factor",
 		ToolTip = "Enables the WebRTC-Video-Pacing field trial and sets the video pacing factor parameter. Larger values are more lenient on larger bitrates. Default: -1.0f (values below zero are discarded.)"
@@ -481,21 +483,21 @@ public:
 	// End WebRTC CVars
 
 	// Begin EditorStreaming CVars
-	static TAutoConsoleVariable<bool> CVarEditorStartOnLaunch;
+	static UE_API TAutoConsoleVariable<bool> CVarEditorStartOnLaunch;
 	UPROPERTY(config, EditAnywhere, Category = "EditorStreaming", meta = (
 		DisplayName = "Start On Launch",
 		ToolTip = "Start Editor Streaming as soon as the Unreal Editor is launched. Default: false"
 		))
 	bool EditorStartOnLaunch = false;
 
-	static TAutoConsoleVariable<bool> CVarEditorUseRemoteSignallingServer;
+	static UE_API TAutoConsoleVariable<bool> CVarEditorUseRemoteSignallingServer;
 	UPROPERTY(config, EditAnywhere, Category = "EditorStreaming", meta = (
 		DisplayName = "Use Remote Signalling Server",
 		ToolTip = "Enables the use of a remote signalling server. Default: false"
 		))
 	bool EditorUseRemoteSignallingServer = false;
 
-	static TAutoConsoleVariable<FString> CVarEditorSource;
+	static UE_API TAutoConsoleVariable<FString> CVarEditorSource;
 	UPROPERTY(config, EditAnywhere, Category = "EditorStreaming", meta = (
 		DisplayName = "Source",
 		ToolTip = "Editor PixelStreaming source. Supported values are `Editor`, `LevelEditorViewport`. Default: Editor"
@@ -504,63 +506,63 @@ public:
 	// End EditorStreaming CVars
 
 	// Begin HMD CVars
-	static TAutoConsoleVariable<bool> CVarHMDEnable;
+	static UE_API TAutoConsoleVariable<bool> CVarHMDEnable;
 	UPROPERTY(config, EditAnywhere, Category = "XR Streaming", meta = (
 		DisplayName = "Enable HMD",
 		ToolTip = "Enables HMD specific functionality for Pixel Streaming. Namely input handling and stereoscopic rendering. Default: false"
 		))
 	bool HMDEnable = false;
 
-	static TAutoConsoleVariable<bool> CVarHMDMatchAspectRatio;
+	static UE_API TAutoConsoleVariable<bool> CVarHMDMatchAspectRatio;
 	UPROPERTY(config, EditAnywhere, Category = "XR Streaming", meta = (
 		DisplayName = "Match Aspect Ratio",
 		ToolTip = "If true automatically resize the rendering resolution to match the aspect ratio determined by the HFoV and VFoV. Default: true"
 		))
 	bool HMDMatchAspectRatio = true;
 
-	static TAutoConsoleVariable<bool> CVarHMDApplyEyePosition;
+	static UE_API TAutoConsoleVariable<bool> CVarHMDApplyEyePosition;
 	UPROPERTY(config, EditAnywhere, Category = "XR Streaming", meta = (
 		DisplayName = "Apply Eye Position",
 		ToolTip = "If true automatically position each eye's rendering by whatever amount WebXR reports for each left-right XRView. If false do no eye positioning. Default: true"
 		))
 	bool HMDAppleEyePosition = true;
 
-	static TAutoConsoleVariable<bool> CVarHMDApplyEyeRotation;
+	static UE_API TAutoConsoleVariable<bool> CVarHMDApplyEyeRotation;
 	UPROPERTY(config, EditAnywhere, Category = "XR Streaming", meta = (
 		DisplayName = "Apply Eye Position",
 		ToolTip = "If true automatically rotate each eye's rendering by whatever amount WebXR reports for each left-right XRView. If false do no eye rotation. Default: true"
 		))
 	bool HMDApplyEyeRotation = true;
 
-	static TAutoConsoleVariable<float> CVarHMDHFOV;
+	static UE_API TAutoConsoleVariable<float> CVarHMDHFOV;
 	UPROPERTY(config, EditAnywhere, Category = "XR Streaming", meta = (
 		DisplayName = "Horizontal FOV Override",
 		ToolTip = "Overrides the horizontal field of view for HMD rendering, values are in degrees and values less than 0.0f disable the override. Default: -1.0f"
 		))
 	float HMDHFOV = -1.0f;
 
-	static TAutoConsoleVariable<float> CVarHMDVFOV;
+	static UE_API TAutoConsoleVariable<float> CVarHMDVFOV;
 	UPROPERTY(config, EditAnywhere, Category = "XR Streaming", meta = (
 		DisplayName = "Vertical FOV Override",
 		ToolTip = "Overrides the vertical field of view for HMD rendering, values are in degrees and values less than 0.0f disable the override. Default: -1.0f"
 		))
 	float HMDVFOV = -1.0f;
 
-	static TAutoConsoleVariable<float> CVarHMDIPD;
+	static UE_API TAutoConsoleVariable<float> CVarHMDIPD;
 	UPROPERTY(config, EditAnywhere, Category = "XR Streaming", meta = (
 		DisplayName = "Interpupillary Distance Override",
 		ToolTip = "Overrides the HMD IPD (interpupillary distance), values are in centimeters and values less than 0.0f disable the override. Default: -1.0f"
 		))
 	float HMDIPD = -1.0f;
 
-	static TAutoConsoleVariable<float> CVarHMDProjectionOffsetX;
+	static UE_API TAutoConsoleVariable<float> CVarHMDProjectionOffsetX;
 	UPROPERTY(config, EditAnywhere, Category = "XR Streaming", meta = (
 		DisplayName = "Horizontal Projection Offset Override",
 		ToolTip = "Overrides the left/right eye projection matrix x-offset, values are in clip space and values less than 0.0f disable the override. Default: -1.0f"
 		))
 	float HMDProjectionOffsetX = -1.0f;
 
-	static TAutoConsoleVariable<float> CVarHMDProjectionOffsetY;
+	static UE_API TAutoConsoleVariable<float> CVarHMDProjectionOffsetY;
 	UPROPERTY(config, EditAnywhere, Category = "XR Streaming", meta = (
 		DisplayName = "Vertical Projection Offset Override",
 		ToolTip = "Overrides the left-right eye projection matrix y-offset, values are in clip space and values less than 0.0f disable the override. Default: -1.0f"
@@ -569,14 +571,14 @@ public:
 	// End HMD CVars
 
 	// Begin Input CVars
-	static TAutoConsoleVariable<bool> CVarInputAllowConsoleCommands;
+	static UE_API TAutoConsoleVariable<bool> CVarInputAllowConsoleCommands;
 	UPROPERTY(config, EditAnywhere, Category = "Input", meta = (
 		DisplayName = "Allow Commands",
 		ToolTip = "If true browser can send consoleCommand payloads that execute in UE's console. Default: false"
 		))
 	bool InputAllowConsoleCommands = false;
 
-	static TAutoConsoleVariable<FString> CVarInputKeyFilter;
+	static UE_API TAutoConsoleVariable<FString> CVarInputKeyFilter;
 	UPROPERTY(config, EditAnywhere, Category = "Input", meta = (
 		DisplayName = "Key Filter",
 		ToolTip = "Comma separated list of keys to ignore from streaming clients. Default: \"\""
@@ -588,19 +590,19 @@ public:
 	// clang-format on
 
 	// Begin UDeveloperSettings Interface
-	virtual FName GetCategoryName() const override;
+	UE_API virtual FName GetCategoryName() const override;
 
 #if WITH_EDITOR
-	virtual FText GetSectionText() const override;
+	UE_API virtual FText GetSectionText() const override;
 
-	virtual void PostEditChangeProperty(FPropertyChangedEvent& PropertyChangedEvent) override;
+	UE_API virtual void PostEditChangeProperty(FPropertyChangedEvent& PropertyChangedEvent) override;
 
-	void VerifyVideoSettings();
+	UE_API void VerifyVideoSettings();
 #endif
 	// End UDeveloperSettings Interface
 
 	// Begin UObject Interface
-	virtual void PostInitProperties() override;
+	UE_API virtual void PostInitProperties() override;
 	// End UObject Interface
 
 	struct FDelegates
@@ -645,28 +647,30 @@ public:
 		FOnEpicRtcLogFilterChanged OnEpicRtcLogFilterChanged;
 	};
 
-	static FDelegates* Delegates();
+	static UE_API FDelegates* Delegates();
 
 private:
-	void SetCVarAndPropertyFromValue(const FString& CVarName, FProperty* Property, const FString& Value);
-	void SetCVarFromProperty(const FString& CVarName, FProperty* Property);
+	UE_API void SetCVarAndPropertyFromValue(const FString& CVarName, FProperty* Property, const FString& Value);
+	UE_API void SetCVarFromProperty(const FString& CVarName, FProperty* Property);
 
-	void InitializeCVarsFromProperties();
-	void ValidateCommandLineArgs();
-	void ParseCommandlineArgs();
-	void ParseLegacyCommandlineArgs();
+	UE_API void InitializeCVarsFromProperties();
+	UE_API void ValidateCommandLineArgs();
+	UE_API void ParseCommandlineArgs();
+	UE_API void ParseLegacyCommandlineArgs();
 
-	static FDelegates* DelegateSingleton;
+	static UE_API FDelegates* DelegateSingleton;
 
 	UFUNCTION()
-	TArray<FString> GetVideoCodecOptions() const;
+	UE_API TArray<FString> GetVideoCodecOptions() const;
 
 	UFUNCTION()
-	TArray<FString> GetScalabilityModeOptions() const;
+	UE_API TArray<FString> GetScalabilityModeOptions() const;
 
 	UFUNCTION()
-	TArray<FString> GetWebRTCCodecPreferencesOptions() const;
+	UE_API TArray<FString> GetWebRTCCodecPreferencesOptions() const;
 
 	UFUNCTION()
-	TArray<FString> GetDefaultStreamerTypeOptions() const;
-};
\ No newline at end of file
+	UE_API TArray<FString> GetDefaultStreamerTypeOptions() const;
+};
+
+#undef UE_API
diff --git a/Engine/Plugins/Media/PixelStreaming2/Source/PixelStreaming2Settings/Private/PixelStreaming2PluginSettings.cpp b/Engine/Plugins/Media/PixelStreaming2/Source/PixelStreaming2Settings/Private/PixelStreaming2PluginSettings.cpp
--- a/Engine/Plugins/Media/PixelStreaming2/Source/PixelStreaming2Settings/Private/PixelStreaming2PluginSettings.cpp
+++ b/Engine/Plugins/Media/PixelStreaming2/Source/PixelStreaming2Settings/Private/PixelStreaming2PluginSettings.cpp
@@ -7,6 +7,8 @@
 #include "Misc/CommandLine.h"
 #include "UObject/ReflectedTypeAccessors.h"
 
+#include UE_INLINE_GENERATED_CPP_BY_NAME(PixelStreaming2PluginSettings)
+
 namespace
 {
 	template <typename TEnumType>
@@ -1436,4 +1438,4 @@ TArray<FString> UPixelStreaming2PluginSettings::GetWebRTCCodecPreferencesOptions
 TArray<FString> UPixelStreaming2PluginSettings::GetDefaultStreamerTypeOptions() const
 {
 	return IPixelStreaming2StreamerFactory::GetAvailableFactoryTypes();
-}
\ No newline at end of file
+}
