[PS2] Fix: Level editor streaming crashing on Mac

#FailedValidations Preflight

diff --git a/Engine/Plugins/Media/PixelStreaming2/Source/PixelStreaming2/Private/VideoProducerViewportBase.cpp b/Engine/Plugins/Media/PixelStreaming2/Source/PixelStreaming2/Private/VideoProducerViewportBase.cpp
--- a/Engine/Plugins/Media/PixelStreaming2/Source/PixelStreaming2/Private/VideoProducerViewportBase.cpp
+++ b/Engine/Plugins/Media/PixelStreaming2/Source/PixelStreaming2/Private/VideoProducerViewportBase.cpp
@@ -2,21 +2,21 @@
 
 #include "VideoProducerViewportBase.h"
 
+#include "EngineModule.h"
+#include "Framework/Application/SlateApplication.h"
 #include "GlobalShader.h"
 #include "Layout/ArrangedChildren.h"
 #include "Layout/Geometry.h"
 #include "Layout/WidgetPath.h"
+#include "MediaShaders.h"
 #include "PixelCaptureInputFrameRHI.h"
+#include "ScreenPass.h"
 
 namespace UE::PixelStreaming2
 {
 	bool FVideoProducerViewportBase::IsFrameAlreadyCopied()
 	{
-#if PLATFORM_MAC
 		return true;
-#else
-		return false;
-#endif
 	}
 	
 	bool FVideoProducerViewportBase::ShouldCaptureViewport()
@@ -56,7 +56,7 @@ namespace UE::PixelStreaming2
 
 	void FVideoProducerViewportBase::OnBackBufferReadyToPresent(SWindow& Window, const FTextureRHIRef& BackBuffer)
 	{	
-		if(!ShouldCaptureViewport())
+		if (!ShouldCaptureViewport())
 		{
 			return;
 		}
@@ -92,18 +92,48 @@ namespace UE::PixelStreaming2
 		{
 			FRDGBuilder GraphBuilder(FRHICommandListImmediate::Get());
 			FRDGTextureRef InputTexture = GraphBuilder.RegisterExternalTexture(CreateRenderTarget(BackBuffer, *Window.GetTitle().ToString()));
-			FRDGTextureRef CroppedTexture = GraphBuilder.CreateTexture(FRDGTextureDesc::Create2D(VirtualSize, InputTexture->Desc.Format, FClearValueBinding::None, GetTexCreateFlags()), TEXT("VideoProducerPIEViewport"));
-	
-			FRHICopyTextureInfo CopyInfo;
-			CopyInfo.SourcePosition = FIntVector(VirtualPos.X, VirtualPos.Y, 0);
-			CopyInfo.Size = FIntVector(VirtualSize.X, VirtualSize.Y, 1);
+			// TODO (william.belcher): The cropped texture having an 8bit pixel format matches the full editor video producer code path but we may need to reconsider this if we want to add 10bit video support
+			FRDGTextureRef CroppedTexture = GraphBuilder.CreateTexture(FRDGTextureDesc::Create2D(VirtualSize, EPixelFormat::PF_B8G8R8A8, FClearValueBinding::None, GetTexCreateFlags()), TEXT("VideoProducerViewportBaseCropped"));
+
+			{
+				// Configure our viewports appropriately
+				FScreenPassTextureViewport InputViewport(InputTexture, FIntRect(VirtualPos, VirtualPos + VirtualSize));
+				FScreenPassTextureViewport OutputViewport(CroppedTexture, FIntRect(FIntPoint(0, 0), VirtualSize));
 
-			AddCopyTexturePass(GraphBuilder, InputTexture, CroppedTexture, CopyInfo);
+				// Rectangle area to use from the source texture
+				const FIntRect ViewRect(VirtualPos, VirtualPos + VirtualSize);
+
+				// Dummy ViewFamily/ViewInfo created to use built in Draw Screen/Texture Pass
+				FSceneViewFamilyContext ViewFamily(FSceneViewFamily::ConstructionValues(nullptr, nullptr, FEngineShowFlags(ESFIM_Game)).SetTime(FGameTime()));
+				FSceneViewInitOptions ViewInitOptions;
+				ViewInitOptions.ViewFamily = &ViewFamily;
+				ViewInitOptions.SetViewRectangle(ViewRect);
+				ViewInitOptions.ViewOrigin = FVector::ZeroVector;
+				ViewInitOptions.ViewRotationMatrix = FMatrix::Identity;
+				ViewInitOptions.ProjectionMatrix = FMatrix::Identity;
+
+				GetRendererModule().CreateAndInitSingleView(GraphBuilder.RHICmdList, &ViewFamily, &ViewInitOptions);
+				const FSceneView &View = *ViewFamily.Views[0];
+				// In cases where texture is converted from a format that doesn't have A channel, we want to force set it to 1.
+				int32 ConversionOperation = 0; // None
+				FModifyAlphaSwizzleRgbaPS::FPermutationDomain PermutationVector;
+				PermutationVector.Set<FModifyAlphaSwizzleRgbaPS::FConversionOp>(ConversionOperation);
+
+				TShaderMapRef<FModifyAlphaSwizzleRgbaPS> PixelShader(GetGlobalShaderMap(GMaxRHIFeatureLevel), PermutationVector);
+				FModifyAlphaSwizzleRgbaPS::FParameters *PixelShaderParameters = GraphBuilder.AllocParameters<FModifyAlphaSwizzleRgbaPS::FParameters>();
+				PixelShaderParameters->InputTexture = InputTexture;
+				PixelShaderParameters->InputSampler = TStaticSamplerState<SF_Point>::GetRHI();
+				PixelShaderParameters->RenderTargets[0] = FRenderTargetBinding{CroppedTexture, ERenderTargetLoadAction::ELoad};
+
+				TShaderMapRef<FScreenPassVS> VertexShader(GetGlobalShaderMap(GMaxRHIFeatureLevel));
+				// Add screen pass to convert whatever format the editor produces to BGRA8
+				AddDrawScreenPass(GraphBuilder, RDG_EVENT_NAME("VideoProducerViewportBaseCropped"), View, OutputViewport, InputViewport, VertexShader, PixelShader, PixelShaderParameters);
+			}
 
 #if PLATFORM_MAC
 			// On Mac specifically, we need to add one more pass to render to a cpu readable texture for AVCodecs
-			FRDGTextureRef StagingTexture = GraphBuilder.CreateTexture(FRDGTextureDesc::Create2D(CroppedTexture->Desc.Extent, EPixelFormat::PF_B8G8R8A8, FClearValueBinding::None, ETextureCreateFlags::CPUReadback), TEXT("VideoProducerPIEViewport MAC Staging"));
-			
+			FRDGTextureRef StagingTexture = GraphBuilder.CreateTexture(FRDGTextureDesc::Create2D(CroppedTexture->Desc.Extent, EPixelFormat::PF_B8G8R8A8, FClearValueBinding::None, ETextureCreateFlags::CPUReadback), TEXT("VideoProducerViewportBaseCropped Mac Staging"));
+
 			AddDrawTexturePass(GraphBuilder, GetGlobalShaderMap(GMaxRHIFeatureLevel), CroppedTexture, StagingTexture, FRDGDrawTextureInfo());
 
 			GraphBuilder.QueueTextureExtraction(StagingTexture, &OutTexture);
