[PS2] Fix: Crash on shutdown. The ApplicationWrapper was keeping the platform application alive for longer than intended leading to an out of order destruction

diff --git a/Engine/Plugins/Media/PixelStreaming2/Source/PixelStreaming2Input/Private/ApplicationWrapper.cpp b/Engine/Plugins/Media/PixelStreaming2/Source/PixelStreaming2Input/Private/ApplicationWrapper.cpp
--- a/Engine/Plugins/Media/PixelStreaming2/Source/PixelStreaming2Input/Private/ApplicationWrapper.cpp
+++ b/Engine/Plugins/Media/PixelStreaming2/Source/PixelStreaming2Input/Private/ApplicationWrapper.cpp
@@ -2,6 +2,30 @@
 
 #include "ApplicationWrapper.h"
 
+#include "Logging.h"
+
+/**
+ * Macro to safely call a method on the wrapped application or return void if invalid.
+ */
+#define SAFE_CALL_WRAPPED_APPLICATION_VOID(Method, ...) \
+	if (!WrappedApplication.IsValid()) \
+	{ \
+		UE_LOGFMT(LogPixelStreaming2Input, Warning, "{FUNCTIONNAME}\nFailed to call {METHODNAME}. Wrapped application is no longer valid", ("FUNCTIONNAME", __FUNCTION__), ("METHODNAME", #Method)); \
+		return; \
+	} \
+	return WrappedApplication->Method(__VA_ARGS__);
+
+/**
+ * Macro to safely call a method on the wrapped application or return DefaultRetVal if invalid.
+ */
+#define SAFE_CALL_WRAPPED_APPLICATION_WITH_RETVAL(DefaultRetVal, Method, ...) \
+	if (!WrappedApplication.IsValid()) \
+	{ \
+		UE_LOGFMT(LogPixelStreaming2Input, Warning, "{FUNCTIONNAME}\nFailed to call {METHODNAME}. Wrapped application is no longer valid", ("FUNCTIONNAME", __FUNCTION__), ("METHODNAME", #Method)); \
+		return DefaultRetVal; \
+	} \
+	return WrappedApplication->Method(__VA_ARGS__);
+
 namespace UE::PixelStreaming2Input
 {
 	FPixelStreaming2ApplicationWrapper::FPixelStreaming2ApplicationWrapper(TSharedPtr<GenericApplication> InWrappedApplication)
@@ -9,11 +33,164 @@ namespace UE::PixelStreaming2Input
 		, WrappedApplication(InWrappedApplication)
 	{
 		InitModifierKeys();
+
+		OnEnginePreExitDelegateHandle = FCoreDelegates::OnEnginePreExit.AddRaw(this, &FPixelStreaming2ApplicationWrapper::OnEnginePreExit);
 	}
 
-	void FPixelStreaming2ApplicationWrapper::SetTargetWindow(TWeakPtr<SWindow> InTargetWindow)
+	FPixelStreaming2ApplicationWrapper::~FPixelStreaming2ApplicationWrapper()
 	{
-		TargetWindow = InTargetWindow;
+		FCoreDelegates::OnEnginePreExit.Remove(OnEnginePreExitDelegateHandle);
+	}
+
+	void FPixelStreaming2ApplicationWrapper::OnEnginePreExit()
+	{
+		// Reset the wrapped application before engine exits to ensure we're not the one holding on to the platform application during destruction
+		WrappedApplication.Reset();
+	}
+
+	void FPixelStreaming2ApplicationWrapper::SetMessageHandler(const TSharedRef<FGenericApplicationMessageHandler>& InMessageHandler)
+	{
+		SAFE_CALL_WRAPPED_APPLICATION_VOID(SetMessageHandler, InMessageHandler)
+	}
+
+	void FPixelStreaming2ApplicationWrapper::PollGameDeviceState(const float TimeDelta)
+	{
+		SAFE_CALL_WRAPPED_APPLICATION_VOID(PollGameDeviceState, TimeDelta)
+	}
+
+	void FPixelStreaming2ApplicationWrapper::PumpMessages(const float TimeDelta)
+	{
+		SAFE_CALL_WRAPPED_APPLICATION_VOID(PumpMessages, TimeDelta)
+	}
+
+	void FPixelStreaming2ApplicationWrapper::ProcessDeferredEvents(const float TimeDelta)
+	{
+		SAFE_CALL_WRAPPED_APPLICATION_VOID(ProcessDeferredEvents, TimeDelta)
+	}
+
+	void FPixelStreaming2ApplicationWrapper::Tick(const float TimeDelta)
+	{
+		SAFE_CALL_WRAPPED_APPLICATION_VOID(Tick, TimeDelta)
+	}
+
+	TSharedRef<FGenericWindow> FPixelStreaming2ApplicationWrapper::MakeWindow()
+	{
+		SAFE_CALL_WRAPPED_APPLICATION_WITH_RETVAL(MakeShared<FGenericWindow>(), MakeWindow)
+	}
+
+	void FPixelStreaming2ApplicationWrapper::InitializeWindow(const TSharedRef<FGenericWindow>& Window, const TSharedRef<FGenericWindowDefinition>& InDefinition, const TSharedPtr<FGenericWindow>& InParent, const bool bShowImmediately)
+	{
+		SAFE_CALL_WRAPPED_APPLICATION_VOID(InitializeWindow, Window, InDefinition, InParent, bShowImmediately)
+	}
+
+	void FPixelStreaming2ApplicationWrapper::SetCapture(const TSharedPtr<FGenericWindow>& InWindow)
+	{
+		SAFE_CALL_WRAPPED_APPLICATION_VOID(SetCapture, InWindow)
+	}
+
+	void* FPixelStreaming2ApplicationWrapper::GetCapture(void) const
+	{
+		SAFE_CALL_WRAPPED_APPLICATION_WITH_RETVAL(nullptr, GetCapture)
+	}
+
+	void FPixelStreaming2ApplicationWrapper::SetHighPrecisionMouseMode(const bool Enable, const TSharedPtr<FGenericWindow>& InWindow)
+	{
+		SAFE_CALL_WRAPPED_APPLICATION_VOID(SetHighPrecisionMouseMode, Enable, InWindow)
+	}
+
+	bool FPixelStreaming2ApplicationWrapper::IsUsingHighPrecisionMouseMode() const
+	{
+		SAFE_CALL_WRAPPED_APPLICATION_WITH_RETVAL(false, IsUsingHighPrecisionMouseMode)
+	}
+
+	bool FPixelStreaming2ApplicationWrapper::IsUsingTrackpad() const
+	{
+		SAFE_CALL_WRAPPED_APPLICATION_WITH_RETVAL(false, IsUsingTrackpad)
+	}
+
+	bool FPixelStreaming2ApplicationWrapper::IsGamepadAttached() const
+	{
+		SAFE_CALL_WRAPPED_APPLICATION_WITH_RETVAL(false, IsGamepadAttached)
+	}
+
+	void FPixelStreaming2ApplicationWrapper::RegisterConsoleCommandListener(const FOnConsoleCommandListener& InListener)
+	{
+		SAFE_CALL_WRAPPED_APPLICATION_VOID(RegisterConsoleCommandListener, InListener)
+	}
+
+	void FPixelStreaming2ApplicationWrapper::AddPendingConsoleCommand(const FString& InCommand)
+	{
+		SAFE_CALL_WRAPPED_APPLICATION_VOID(AddPendingConsoleCommand, InCommand)
+	}
+
+	FPlatformRect FPixelStreaming2ApplicationWrapper::GetWorkArea(const FPlatformRect& CurrentWindow) const
+	{
+		SAFE_CALL_WRAPPED_APPLICATION_WITH_RETVAL(FPlatformRect(), GetWorkArea, CurrentWindow)
+	}
+
+	bool FPixelStreaming2ApplicationWrapper::TryCalculatePopupWindowPosition(const FPlatformRect& InAnchor, const FVector2D& InSize, const FVector2D& ProposedPlacement, const EPopUpOrientation::Type Orientation, /*OUT*/ FVector2D* const CalculatedPopUpPosition) const
+	{
+		SAFE_CALL_WRAPPED_APPLICATION_WITH_RETVAL(false, TryCalculatePopupWindowPosition, InAnchor, InSize, ProposedPlacement, Orientation, CalculatedPopUpPosition)
+	}
+
+	void FPixelStreaming2ApplicationWrapper::GetInitialDisplayMetrics(FDisplayMetrics& OutDisplayMetrics) const
+	{
+		SAFE_CALL_WRAPPED_APPLICATION_VOID(GetInitialDisplayMetrics, OutDisplayMetrics)
+	}
+
+	EWindowTitleAlignment::Type FPixelStreaming2ApplicationWrapper::GetWindowTitleAlignment() const
+	{
+		SAFE_CALL_WRAPPED_APPLICATION_WITH_RETVAL(EWindowTitleAlignment::Left, GetWindowTitleAlignment)
+	}
+
+	EWindowTransparency FPixelStreaming2ApplicationWrapper::GetWindowTransparencySupport() const
+	{
+		SAFE_CALL_WRAPPED_APPLICATION_WITH_RETVAL(EWindowTransparency::None, GetWindowTransparencySupport)
+	}
+
+	void FPixelStreaming2ApplicationWrapper::DestroyApplication()
+	{
+		SAFE_CALL_WRAPPED_APPLICATION_VOID(DestroyApplication)
+	}
+
+	IInputInterface* FPixelStreaming2ApplicationWrapper::GetInputInterface()
+	{
+		SAFE_CALL_WRAPPED_APPLICATION_WITH_RETVAL(nullptr, GetInputInterface)
+	}
+
+	ITextInputMethodSystem* FPixelStreaming2ApplicationWrapper::GetTextInputMethodSystem()
+	{
+		SAFE_CALL_WRAPPED_APPLICATION_WITH_RETVAL(nullptr, GetTextInputMethodSystem)
+	}
+	
+	void FPixelStreaming2ApplicationWrapper::SendAnalytics(IAnalyticsProvider* Provider)
+	{
+		SAFE_CALL_WRAPPED_APPLICATION_VOID(SendAnalytics, Provider)
+	}
+
+	bool FPixelStreaming2ApplicationWrapper::SupportsSystemHelp() const
+	{
+		SAFE_CALL_WRAPPED_APPLICATION_WITH_RETVAL(false, SupportsSystemHelp)
+	}
+
+	void FPixelStreaming2ApplicationWrapper::ShowSystemHelp()
+	{
+		SAFE_CALL_WRAPPED_APPLICATION_VOID(ShowSystemHelp)
+	}
+
+	bool FPixelStreaming2ApplicationWrapper::ApplicationLicenseValid(FPlatformUserId PlatformUser)
+	{
+		SAFE_CALL_WRAPPED_APPLICATION_WITH_RETVAL(false, ApplicationLicenseValid, PlatformUser)
+	}
+
+	bool FPixelStreaming2ApplicationWrapper::IsMouseAttached() const
+	{
+		return true;
+	}
+
+	bool FPixelStreaming2ApplicationWrapper::IsCursorDirectlyOverSlateWindow() const
+	{
+		return true;
 	}
 
 	TSharedPtr<FGenericWindow> FPixelStreaming2ApplicationWrapper::GetWindowUnderCursor()
@@ -34,7 +211,12 @@ namespace UE::PixelStreaming2Input
 			}
 		}
 
-		return WrappedApplication->GetWindowUnderCursor();
+		SAFE_CALL_WRAPPED_APPLICATION_WITH_RETVAL(nullptr, GetWindowUnderCursor)
+	}
+
+	void FPixelStreaming2ApplicationWrapper::SetTargetWindow(TWeakPtr<SWindow> InTargetWindow)
+	{
+		TargetWindow = InTargetWindow;
 	}
 
 	/** Initialize the list of possible modifier keys. */
@@ -86,4 +268,7 @@ namespace UE::PixelStreaming2Input
 			/*bInAreCapsLocked*/ ModifierKeys[EModifierKey::CapsLock].bActive);
 		return ModifierKeysState;
 	}
-} // namespace UE::PixelStreaming2Input
\ No newline at end of file
+} // namespace UE::PixelStreaming2Input
+
+#undef SAFE_CALL_WRAPPED_APPLICATION_VOID
+#undef SAFE_CALL_WRAPPED_APPLICATION_WITH_RETVAL
\ No newline at end of file
diff --git a/Engine/Plugins/Media/PixelStreaming2/Source/PixelStreaming2Input/Private/ApplicationWrapper.h b/Engine/Plugins/Media/PixelStreaming2/Source/PixelStreaming2Input/Private/ApplicationWrapper.h
--- a/Engine/Plugins/Media/PixelStreaming2/Source/PixelStreaming2Input/Private/ApplicationWrapper.h
+++ b/Engine/Plugins/Media/PixelStreaming2/Source/PixelStreaming2Input/Private/ApplicationWrapper.h
@@ -15,44 +15,46 @@ namespace UE::PixelStreaming2Input
 	{
 	public:
 		FPixelStreaming2ApplicationWrapper(TSharedPtr<GenericApplication> InWrappedApplication);
+		~FPixelStreaming2ApplicationWrapper();
+
+		void OnEnginePreExit();
 
 		/**
 		 * Functions passed directly to the wrapped application.
 		 */
-		virtual void					   SetMessageHandler(const TSharedRef<FGenericApplicationMessageHandler>& InMessageHandler) { WrappedApplication->SetMessageHandler(InMessageHandler); }
-		virtual void					   PollGameDeviceState(const float TimeDelta) { WrappedApplication->PollGameDeviceState(TimeDelta); }
-		virtual void					   PumpMessages(const float TimeDelta) { WrappedApplication->PumpMessages(TimeDelta); }
-		virtual void					   ProcessDeferredEvents(const float TimeDelta) { WrappedApplication->ProcessDeferredEvents(TimeDelta); }
-		virtual void					   Tick(const float TimeDelta) { WrappedApplication->Tick(TimeDelta); }
-		virtual TSharedRef<FGenericWindow> MakeWindow() { return WrappedApplication->MakeWindow(); }
-		virtual void					   InitializeWindow(const TSharedRef<FGenericWindow>& Window, const TSharedRef<FGenericWindowDefinition>& InDefinition, const TSharedPtr<FGenericWindow>& InParent, const bool bShowImmediately) { WrappedApplication->InitializeWindow(Window, InDefinition, InParent, bShowImmediately); }
-		virtual void					   SetCapture(const TSharedPtr<FGenericWindow>& InWindow) { WrappedApplication->SetCapture(InWindow); }
-		virtual void* GetCapture(void) const { return WrappedApplication->GetCapture(); }
-		// virtual FModifierKeysState			GetModifierKeys() const { return WrappedApplication->GetModifierKeys(); }
-		virtual void						SetHighPrecisionMouseMode(const bool Enable, const TSharedPtr<FGenericWindow>& InWindow) { WrappedApplication->SetHighPrecisionMouseMode(Enable, InWindow); };
-		virtual bool						IsUsingHighPrecisionMouseMode() const { return WrappedApplication->IsUsingHighPrecisionMouseMode(); }
-		virtual bool						IsUsingTrackpad() const { return WrappedApplication->IsUsingTrackpad(); }
-		virtual bool						IsGamepadAttached() const { return WrappedApplication->IsGamepadAttached(); }
-		virtual void						RegisterConsoleCommandListener(const FOnConsoleCommandListener& InListener) { WrappedApplication->RegisterConsoleCommandListener(InListener); }
-		virtual void						AddPendingConsoleCommand(const FString& InCommand) { WrappedApplication->AddPendingConsoleCommand(InCommand); }
-		virtual FPlatformRect				GetWorkArea(const FPlatformRect& CurrentWindow) const { return WrappedApplication->GetWorkArea(CurrentWindow); }
-		virtual bool						TryCalculatePopupWindowPosition(const FPlatformRect& InAnchor, const FVector2D& InSize, const FVector2D& ProposedPlacement, const EPopUpOrientation::Type Orientation, /*OUT*/ FVector2D* const CalculatedPopUpPosition) const { return WrappedApplication->TryCalculatePopupWindowPosition(InAnchor, InSize, ProposedPlacement, Orientation, CalculatedPopUpPosition); }
-		virtual void						GetInitialDisplayMetrics(FDisplayMetrics& OutDisplayMetrics) const { WrappedApplication->GetInitialDisplayMetrics(OutDisplayMetrics); }
-		virtual EWindowTitleAlignment::Type GetWindowTitleAlignment() const { return WrappedApplication->GetWindowTitleAlignment(); }
-		virtual EWindowTransparency			GetWindowTransparencySupport() const { return WrappedApplication->GetWindowTransparencySupport(); }
-		virtual void						DestroyApplication() { WrappedApplication->DestroyApplication(); }
-		virtual IInputInterface* GetInputInterface() { return WrappedApplication->GetInputInterface(); }
-		virtual ITextInputMethodSystem* GetTextInputMethodSystem() { return WrappedApplication->GetTextInputMethodSystem(); }
-		virtual void						SendAnalytics(IAnalyticsProvider* Provider) { WrappedApplication->SendAnalytics(Provider); }
-		virtual bool						SupportsSystemHelp() const { return WrappedApplication->SupportsSystemHelp(); }
-		virtual void						ShowSystemHelp() { WrappedApplication->ShowSystemHelp(); }
-		virtual bool						ApplicationLicenseValid(FPlatformUserId PlatformUser = PLATFORMUSERID_NONE) { return WrappedApplication->ApplicationLicenseValid(PlatformUser); }
+		virtual void						SetMessageHandler(const TSharedRef<FGenericApplicationMessageHandler>& InMessageHandler) override;
+		virtual void						PollGameDeviceState(const float TimeDelta) override;
+		virtual void						PumpMessages(const float TimeDelta) override;
+		virtual void						ProcessDeferredEvents(const float TimeDelta) override;
+		virtual void						Tick(const float TimeDelta) override;
+		virtual TSharedRef<FGenericWindow>	MakeWindow() override;
+		virtual void						InitializeWindow(const TSharedRef<FGenericWindow>& Window, const TSharedRef<FGenericWindowDefinition>& InDefinition, const TSharedPtr<FGenericWindow>& InParent, const bool bShowImmediately) override;
+		virtual void						SetCapture(const TSharedPtr<FGenericWindow>& InWindow) override;
+		virtual void*						GetCapture(void) const override;
+		virtual void						SetHighPrecisionMouseMode(const bool Enable, const TSharedPtr<FGenericWindow>& InWindow) override;
+		virtual bool						IsUsingHighPrecisionMouseMode() const override;
+		virtual bool						IsUsingTrackpad() const override;
+		virtual bool						IsGamepadAttached() const override;
+		virtual void						RegisterConsoleCommandListener(const FOnConsoleCommandListener& InListener) override;
+		virtual void						AddPendingConsoleCommand(const FString& InCommand) override;
+		virtual FPlatformRect				GetWorkArea(const FPlatformRect& CurrentWindow) const override;
+		virtual bool						TryCalculatePopupWindowPosition(const FPlatformRect& InAnchor, const FVector2D& InSize, const FVector2D& ProposedPlacement, const EPopUpOrientation::Type Orientation, /*OUT*/ FVector2D* const CalculatedPopUpPosition) const override;
+		virtual void						GetInitialDisplayMetrics(FDisplayMetrics& OutDisplayMetrics) const override;
+		virtual EWindowTitleAlignment::Type GetWindowTitleAlignment() const override;
+		virtual EWindowTransparency			GetWindowTransparencySupport() const override;
+		virtual void						DestroyApplication() override;
+		virtual IInputInterface*			GetInputInterface() override;
+		virtual ITextInputMethodSystem*		GetTextInputMethodSystem() override;
+		virtual void						SendAnalytics(IAnalyticsProvider* Provider) override;
+		virtual bool						SupportsSystemHelp() const override;
+		virtual void						ShowSystemHelp() override;
+		virtual bool						ApplicationLicenseValid(FPlatformUserId PlatformUser = PLATFORMUSERID_NONE) override;
 
 		/**
 		 * Functions with overridden behavior.
 		 */
-		virtual bool					   IsMouseAttached() const { return true; }
-		virtual bool					   IsCursorDirectlyOverSlateWindow() const { return true; }
+		virtual bool					   IsMouseAttached() const override;
+		virtual bool					   IsCursorDirectlyOverSlateWindow() const override;
 		virtual TSharedPtr<FGenericWindow> GetWindowUnderCursor() override;
 
 		/**
@@ -103,6 +105,9 @@ namespace UE::PixelStreaming2Input
 		 * @return The current set of active modifier keys.
 		 */
 		virtual FModifierKeysState GetModifierKeys() const;
+
+	private:
+		FDelegateHandle OnEnginePreExitDelegateHandle;
 	};
 
 	/**
diff --git a/Engine/Plugins/Media/PixelStreaming2/Source/PixelStreaming2Input/Private/InputHandler.cpp b/Engine/Plugins/Media/PixelStreaming2/Source/PixelStreaming2Input/Private/InputHandler.cpp
--- a/Engine/Plugins/Media/PixelStreaming2/Source/PixelStreaming2Input/Private/InputHandler.cpp
+++ b/Engine/Plugins/Media/PixelStreaming2/Source/PixelStreaming2Input/Private/InputHandler.cpp
@@ -301,8 +301,16 @@ namespace UE::PixelStreaming2Input
 				FSlateApplication::Get().OnCursorSet();
 				// Make sure the application is active.
 				FSlateApplication::Get().ProcessApplicationActivationEvent(true);
-
-				FVector2D OldCursorLocation = PixelStreamerApplicationWrapper->WrappedApplication->Cursor->GetPosition();
+				// Default to the current touch location so that if the wrapped application is invalid, we just don't move the cursor
+				FVector2D OldCursorLocation = FVector2D(TouchLocation.X, TouchLocation.Y);
+				if (PixelStreamerApplicationWrapper->WrappedApplication.IsValid())
+				{
+					OldCursorLocation = PixelStreamerApplicationWrapper->WrappedApplication->Cursor->GetPosition();	
+				} 
+				else
+				{
+					UE_LOG(LogPixelStreaming2Input, Warning, TEXT("Wrapped application is no longer valid"));
+				}
 				PixelStreamerApplicationWrapper->Cursor->SetPosition(OldCursorLocation.X, OldCursorLocation.Y);
 				FSlateApplication::Get().OverridePlatformApplication(PixelStreamerApplicationWrapper);
 			}
@@ -313,7 +321,14 @@ namespace UE::PixelStreaming2Input
 			TSharedPtr<FSlateUser> User = FSlateApplication::Get().GetCursorUser();
 			User->SetCursorPosition(TouchLocation);
 			PixelStreamerApplicationWrapper->Cursor->SetPosition(TouchLocation.X, TouchLocation.Y);
-			PixelStreamerApplicationWrapper->WrappedApplication->Cursor->SetPosition(TouchLocation.X, TouchLocation.Y);
+			if (PixelStreamerApplicationWrapper->WrappedApplication.IsValid())
+			{
+				PixelStreamerApplicationWrapper->WrappedApplication->Cursor->SetPosition(TouchLocation.X, TouchLocation.Y);
+			} 
+			else
+			{
+				UE_LOG(LogPixelStreaming2Input, Warning, TEXT("Wrapped application is no longer valid"));
+			}
 
 			bHandled = MessageHandler->OnTouchStarted(PixelStreamerApplicationWrapper->GetWindowUnderCursor(), TouchLocation, TouchForce, TouchIndex, 0); // TODO: ControllerId?
 		}
@@ -395,8 +410,15 @@ namespace UE::PixelStreaming2Input
 		if (NumActiveTouches == 0 && !bIsMouseActive && InputType == EPixelStreaming2InputType::RouteToWindow)
 		{
 			FVector2D OldCursorLocation = PixelStreamerApplicationWrapper->Cursor->GetPosition();
-			PixelStreamerApplicationWrapper->WrappedApplication->Cursor->SetPosition(OldCursorLocation.X, OldCursorLocation.Y);
-			FSlateApplication::Get().OverridePlatformApplication(PixelStreamerApplicationWrapper->WrappedApplication);
+			if (PixelStreamerApplicationWrapper->WrappedApplication.IsValid())
+			{
+				PixelStreamerApplicationWrapper->WrappedApplication->Cursor->SetPosition(OldCursorLocation.X, OldCursorLocation.Y);
+				FSlateApplication::Get().OverridePlatformApplication(PixelStreamerApplicationWrapper->WrappedApplication);
+			} 
+			else
+			{
+				UE_LOG(LogPixelStreaming2Input, Warning, TEXT("Wrapped application is no longer valid"));
+			}
 		}
 
 		return bHandled;
@@ -503,8 +525,15 @@ namespace UE::PixelStreaming2Input
 	{
 		if (NumActiveTouches == 0)
 		{
-			// Restore normal application layer if there is no active touches and MouseEnter hasn't been triggered
-			FSlateApplication::Get().OverridePlatformApplication(PixelStreamerApplicationWrapper->WrappedApplication);
+			if (PixelStreamerApplicationWrapper->WrappedApplication.IsValid())
+			{
+				// Restore normal application layer if there is no active touches and MouseEnter hasn't been triggered
+				FSlateApplication::Get().OverridePlatformApplication(PixelStreamerApplicationWrapper->WrappedApplication);
+			} 
+			else
+			{
+				UE_LOG(LogPixelStreaming2Input, Warning, TEXT("Wrapped application is no longer valid"));
+			}
 		}
 		bIsMouseActive = false;
 		UE_LOG(LogPixelStreaming2Input, Verbose, TEXT("MOUSE_LEAVE"));
diff --git a/Engine/Plugins/Media/PixelStreaming2/Source/PixelStreaming2Input/Private/Logging.h b/Engine/Plugins/Media/PixelStreaming2/Source/PixelStreaming2Input/Private/Logging.h
--- a/Engine/Plugins/Media/PixelStreaming2/Source/PixelStreaming2Input/Private/Logging.h
+++ b/Engine/Plugins/Media/PixelStreaming2/Source/PixelStreaming2Input/Private/Logging.h
@@ -3,5 +3,6 @@
 #pragma once
 
 #include "Logging/LogMacros.h"
+#include "Logging/StructuredLog.h"
 
-DECLARE_LOG_CATEGORY_EXTERN(LogPixelStreaming2Input, Log, All);
\ No newline at end of file
+DECLARE_LOG_CATEGORY_EXTERN(LogPixelStreaming2Input, Log, All);
