[PS2] QOL: Update Pixel Streaming 2 toolbar to display "Pixel Streaming 2"

#FailedValidations Preflight

diff --git a/Engine/Plugins/Media/PixelStreaming2/Source/PixelStreaming2Editor/Private/PixelStreaming2Toolbar.cpp b/Engine/Plugins/Media/PixelStreaming2/Source/PixelStreaming2Editor/Private/PixelStreaming2Toolbar.cpp
--- a/Engine/Plugins/Media/PixelStreaming2/Source/PixelStreaming2Editor/Private/PixelStreaming2Toolbar.cpp
+++ b/Engine/Plugins/Media/PixelStreaming2/Source/PixelStreaming2Editor/Private/PixelStreaming2Toolbar.cpp
@@ -338,7 +338,7 @@ namespace UE::EditorPixelStreaming2
 
 								return MenuBuilder.MakeWidget();
 							}),
-						LOCTEXT("PixelStreaming2Menu", "Pixel Streaming"),
+						LOCTEXT("PixelStreaming2Menu", "Pixel Streaming 2"),
 						LOCTEXT("PixelStreaming2MenuTooltip", "Configure Pixel Streaming"),
 						FSlateIcon(FPixelStreaming2Style::GetStyleSetName(), "PixelStreaming2.Icon"),
 						false,
