[PS2] Fix: Internal signaling server not auto starting during PIE. Added project setting to disable auto PIE streaming

diff --git a/Engine/Plugins/Media/PixelStreaming2/Source/PixelStreaming2Editor/Private/PixelStreaming2EditorModule.cpp b/Engine/Plugins/Media/PixelStreaming2/Source/PixelStreaming2Editor/Private/PixelStreaming2EditorModule.cpp
--- a/Engine/Plugins/Media/PixelStreaming2/Source/PixelStreaming2Editor/Private/PixelStreaming2EditorModule.cpp
+++ b/Engine/Plugins/Media/PixelStreaming2/Source/PixelStreaming2Editor/Private/PixelStreaming2EditorModule.cpp
@@ -144,6 +144,7 @@ namespace UE::EditorPixelStreaming2
 		}
 
 		StopStreaming();
+		StopSignalling(true);
 	}
 
 	void FPixelStreaming2EditorModule::InitEditorStreaming(IPixelStreaming2RTCModule& Module)
@@ -314,23 +315,12 @@ namespace UE::EditorPixelStreaming2
 
 	void FPixelStreaming2EditorModule::StopStreaming()
 	{
-		if (!EditorStreamer.IsValid())
-		{
-			return;
-		}
-
-		if (!UPixelStreaming2PluginSettings::CVarEditorUseRemoteSignallingServer.GetValueOnAnyThread())
+		if (EditorStreamer)
 		{
-			StopSignalling();
+			EditorStreamer->StopStreaming();
 		}
-
-		if (TSharedPtr<IPixelStreaming2InputHandler> InputHandler = EditorStreamer->GetInputHandler().Pin())
-		{
-			InputHandler->SetTargetViewport(nullptr);
-			InputHandler->SetTargetWindow(nullptr);
-		}
-
-		EditorStreamer->StopStreaming();
+		
+		StopSignalling(false);
 	}
 
 	void FPixelStreaming2EditorModule::StartSignalling()
@@ -350,7 +340,7 @@ namespace UE::EditorPixelStreaming2
 			if (DownloadProcess.IsValid())
 			{
 				DownloadProcess->OnCompleted().BindLambda([this](int ExitCode) {
-					StopSignalling();
+					StopSignalling(true);
 					StartSignalling();
 				});
 
@@ -372,13 +362,23 @@ namespace UE::EditorPixelStreaming2
 		SignallingServer->Launch(LaunchArgs);
 	}
 
-	void FPixelStreaming2EditorModule::StopSignalling()
+	void FPixelStreaming2EditorModule::StopSignalling(bool bForce)
 	{
-		if (SignallingServer.IsValid())
+		if (!SignallingServer.IsValid() || !SignallingServer->HasLaunched())
+		{
+			return;
+		}
+
+		// Force stop and reset the signalling server if desired
+		if (bForce)
 		{
 			SignallingServer->Stop();
 			SignallingServer.Reset();
+			return;
 		}
+
+		// If we're not force stopping signalling, we delay 2 frames to make sure all disconnect messages have been received
+		StopSignallingAfterFrames(2);
 	}
 
 	TSharedPtr<UE::PixelStreaming2Servers::IServer> FPixelStreaming2EditorModule::GetSignallingServer()
@@ -414,6 +414,7 @@ namespace UE::EditorPixelStreaming2
 	{
 		SSLCertificatePath = Path;
 	}
+
 	void FPixelStreaming2EditorModule::SetSSLPrivateKeyPath(const FString& Path)
 	{
 		SSLPrivateKeyPath = Path;
@@ -532,9 +533,44 @@ namespace UE::EditorPixelStreaming2
 		}
 	}
 
+	void FPixelStreaming2EditorModule::StopSignallingAfterFrames(uint8 FrameDelay)
+	{
+		// NOTE (william.belcher): Because the signalling server websockets are ticked, we have to wait two frames before calling StopSignalling as
+		// otherwise the server will still report having this PIEStreamer connected
+		TWeakPtr<UE::PixelStreaming2Servers::IServer> WeakServer = SignallingServer; // Use a weak pointer to the signalling server to track module validity
+		uint8 NumFramesToWait = FrameDelay;
+		FTSTicker::GetCoreTicker().AddTicker(FTickerDelegate::CreateLambda([this, WeakServer, NumFramesToWait](float) mutable {
+			NumFramesToWait--;
+			if (NumFramesToWait > 0)
+			{
+				// If we still have frames to wait, return true to continue ticking
+				return true;
+			}
+
+			if (WeakServer.IsValid())
+			{
+				// A `this` capture is safe here because the lambda is called synchronously before the function returns
+				this->SignallingServer->GetNumStreamers([this](uint16 NumStreamers) mutable {
+					if (NumStreamers == 0)
+					{
+						this->SignallingServer->Stop();
+						this->SignallingServer.Reset();
+					}
+				});
+			}
+
+			return false;
+		}));
+	}
+
 #if WITH_EDITOR
 	void FPixelStreaming2EditorModule::OnBeginPIE(bool bIsSimulating)
 	{
+		if (!UPixelStreaming2PluginSettings::CVarAutoStreamPIE.GetValueOnAnyThread())
+		{
+			return;
+		}
+
 		IPixelStreaming2Module& Module = IPixelStreaming2Module::Get();
 		PIEStreamer = Module.CreateStreamer(Module.GetDefaultStreamerID());
 		// Give the PIE streamer the default url if the user hasn't specified one when launching the editor
@@ -544,6 +580,11 @@ namespace UE::EditorPixelStreaming2
 			PIEStreamer->SetConnectionURL(FString::Printf(TEXT("%s:%d"), *SignallingDomain, StreamerPort));
 		}
 
+		if (!UPixelStreaming2PluginSettings::CVarEditorUseRemoteSignallingServer.GetValueOnAnyThread())
+		{
+			StartSignalling();
+		}
+
 		// Bind to start/stop streaming so we disable/restore relevant editor settings
 		PIEStreamer->OnStreamingStarted().AddLambda([this](IPixelStreaming2Streamer* Streamer) {
 			DisableCPUThrottlingSetting();
@@ -592,14 +633,15 @@ namespace UE::EditorPixelStreaming2
 
 	void FPixelStreaming2EditorModule::OnEndPIE(bool bIsSimulating)
 	{
-		if (!PIEStreamer)
+		if (PIEStreamer)
 		{
-			return;
+			PIEStreamer->StopStreaming();
+			
+			// Remove from the module so that the PIE streamer doesn't pollute the toolbar
+			IPixelStreaming2Module::Get().DeleteStreamer(PIEStreamer);
 		}
-		
-		PIEStreamer->StopStreaming();
 
-		IPixelStreaming2Module::Get().DeleteStreamer(PIEStreamer);
+		StopSignalling(false);
 	}
 #endif
 } // namespace UE::EditorPixelStreaming2
diff --git a/Engine/Plugins/Media/PixelStreaming2/Source/PixelStreaming2Editor/Private/PixelStreaming2EditorModule.h b/Engine/Plugins/Media/PixelStreaming2/Source/PixelStreaming2Editor/Private/PixelStreaming2EditorModule.h
--- a/Engine/Plugins/Media/PixelStreaming2/Source/PixelStreaming2Editor/Private/PixelStreaming2EditorModule.h
+++ b/Engine/Plugins/Media/PixelStreaming2/Source/PixelStreaming2Editor/Private/PixelStreaming2EditorModule.h
@@ -22,7 +22,7 @@ namespace UE::EditorPixelStreaming2
 		virtual void StopStreaming() override;
 
 		virtual void											StartSignalling() override;
-		virtual void											StopSignalling() override;
+		virtual void											StopSignalling(bool bForce) override;
 		virtual TSharedPtr<UE::PixelStreaming2Servers::IServer> GetSignallingServer() override;
 
 		virtual void	SetSignallingDomain(const FString& InSignallingDomain) override;
@@ -47,6 +47,7 @@ namespace UE::EditorPixelStreaming2
 		void MaybeResizeEditor(TSharedPtr<SWindow> RootWindow);
 		void DisableCPUThrottlingSetting();
 		void RestoreCPUThrottlingSetting(bool bForce = false);
+		void StopSignallingAfterFrames(uint8 FrameDelay);
 
 #if WITH_EDITOR
 		void OnBeginPIE(bool bIsSimulating);
diff --git a/Engine/Plugins/Media/PixelStreaming2/Source/PixelStreaming2Editor/Private/PixelStreaming2Toolbar.cpp b/Engine/Plugins/Media/PixelStreaming2/Source/PixelStreaming2Editor/Private/PixelStreaming2Toolbar.cpp
--- a/Engine/Plugins/Media/PixelStreaming2/Source/PixelStreaming2Editor/Private/PixelStreaming2Toolbar.cpp
+++ b/Engine/Plugins/Media/PixelStreaming2/Source/PixelStreaming2Editor/Private/PixelStreaming2Toolbar.cpp
@@ -127,7 +127,7 @@ namespace UE::EditorPixelStreaming2
 			FPixelStreaming2Commands::Get().ExternalSignalling,
 			FExecuteAction::CreateLambda([]() {
 				SetUseRemoteSignallingServer(!GetUseRemoteSignallingServer());
-				IPixelStreaming2EditorModule::Get().StopSignalling();
+				IPixelStreaming2EditorModule::Get().StopSignalling(true);
 			}),
 			FCanExecuteAction::CreateLambda([] {
 				TSharedPtr<UE::PixelStreaming2Servers::IServer> SignallingServer = IPixelStreaming2EditorModule::Get().GetSignallingServer();
@@ -203,7 +203,7 @@ namespace UE::EditorPixelStreaming2
 		PluginCommands->MapAction(
 			FPixelStreaming2Commands::Get().StopSignalling,
 			FExecuteAction::CreateLambda([]() {
-				IPixelStreaming2EditorModule::Get().StopSignalling();
+				IPixelStreaming2EditorModule::Get().StopSignalling(true);
 			}),
 			FCanExecuteAction::CreateLambda([] {
 				TSharedPtr<UE::PixelStreaming2Servers::IServer> SignallingServer = IPixelStreaming2EditorModule::Get().GetSignallingServer();
diff --git a/Engine/Plugins/Media/PixelStreaming2/Source/PixelStreaming2Editor/Public/IPixelStreaming2EditorModule.h b/Engine/Plugins/Media/PixelStreaming2/Source/PixelStreaming2Editor/Public/IPixelStreaming2EditorModule.h
--- a/Engine/Plugins/Media/PixelStreaming2/Source/PixelStreaming2Editor/Public/IPixelStreaming2EditorModule.h
+++ b/Engine/Plugins/Media/PixelStreaming2/Source/PixelStreaming2Editor/Public/IPixelStreaming2EditorModule.h
@@ -55,8 +55,9 @@ public:
 	/**
 	 * Stop the inbuilt C++ signalling server
 	 *
+	 * @param bForce If true, will stop the server even if there are active streamers connected. Default is true.
 	 */
-	virtual void StopSignalling() = 0;
+	virtual void StopSignalling(bool bForce = true) = 0;
 
 	/**
 	 * Get the inbuilt C++ signalling server
diff --git a/Engine/Plugins/Media/PixelStreaming2/Source/PixelStreaming2Settings/Internal/PixelStreaming2PluginSettings.h b/Engine/Plugins/Media/PixelStreaming2/Source/PixelStreaming2Settings/Internal/PixelStreaming2PluginSettings.h
--- a/Engine/Plugins/Media/PixelStreaming2/Source/PixelStreaming2Settings/Internal/PixelStreaming2PluginSettings.h
+++ b/Engine/Plugins/Media/PixelStreaming2/Source/PixelStreaming2Settings/Internal/PixelStreaming2PluginSettings.h
@@ -483,13 +483,6 @@ public:
 	// End WebRTC CVars
 
 	// Begin EditorStreaming CVars
-	static UE_API TAutoConsoleVariable<bool> CVarEditorStartOnLaunch;
-	UPROPERTY(config, EditAnywhere, Category = "EditorStreaming", meta = (
-		DisplayName = "Start On Launch",
-		ToolTip = "Start Editor Streaming as soon as the Unreal Editor is launched. Default: false"
-		))
-	bool EditorStartOnLaunch = false;
-
 	static UE_API TAutoConsoleVariable<bool> CVarEditorUseRemoteSignallingServer;
 	UPROPERTY(config, EditAnywhere, Category = "EditorStreaming", meta = (
 		DisplayName = "Use Remote Signalling Server",
@@ -497,9 +490,23 @@ public:
 		))
 	bool EditorUseRemoteSignallingServer = false;
 
+	static UE_API TAutoConsoleVariable<bool> CVarAutoStreamPIE;
+	UPROPERTY(config, EditAnywhere, Category = "EditorStreaming", meta = (
+		DisplayName = "Auto Stream Play In Editor",
+		ToolTip = "Automatically start streaming when playing in the editor. Default: true"
+		))
+	bool AutoStreamPlayInEditor = true;
+
+	static UE_API TAutoConsoleVariable<bool> CVarEditorStartOnLaunch;
+	UPROPERTY(config, EditAnywhere, Category = "EditorStreaming", meta = (
+		DisplayName = "Start Editor Streamer On Launch",
+		ToolTip = "Start Editor Streaming as soon as the Unreal Editor is launched. Default: false"
+		))
+	bool EditorStartOnLaunch = false;
+
 	static UE_API TAutoConsoleVariable<FString> CVarEditorSource;
 	UPROPERTY(config, EditAnywhere, Category = "EditorStreaming", meta = (
-		DisplayName = "Source",
+		DisplayName = "Editor Streamer Source",
 		ToolTip = "Editor PixelStreaming source. Supported values are `Editor`, `LevelEditorViewport`. Default: Editor"
 		))
 	EPixelStreaming2EditorStreamTypes EditorSource = EPixelStreaming2EditorStreamTypes::Editor;
diff --git a/Engine/Plugins/Media/PixelStreaming2/Source/PixelStreaming2Settings/Private/PixelStreaming2PluginSettings.cpp b/Engine/Plugins/Media/PixelStreaming2/Source/PixelStreaming2Settings/Private/PixelStreaming2PluginSettings.cpp
--- a/Engine/Plugins/Media/PixelStreaming2/Source/PixelStreaming2Settings/Private/PixelStreaming2PluginSettings.cpp
+++ b/Engine/Plugins/Media/PixelStreaming2/Source/PixelStreaming2Settings/Private/PixelStreaming2PluginSettings.cpp
@@ -208,6 +208,7 @@ static const TSet<TPair<FString, FString>> GetCmdArg = {
 	{ "PixelStreaming2.WebRTC.VideoPacing.MaxDelay", "WebRTCVideoPacingMaxDelay" },
 	{ "PixelStreaming2.WebRTC.VideoPacing.Factor", "WebRTCVideoPacingFactor" },
 	{ "PixelStreaming2.Editor.StartOnLaunch", "EditorStartOnLaunch" },
+	{ "PixelStreaming2.Editor.AutoStreamPIE", "AutoStreamPlayInEditor" },
 	{ "PixelStreaming2.Editor.UseRemoteSignallingServer", "EditorUseRemoteSignallingServer" },
 	{ "PixelStreaming2.HMD.Enable", "HMDEnable" },
 	{ "PixelStreaming2.HMD.MatchAspectRatio", "HMDMatchAspectRatio" },
@@ -540,6 +541,12 @@ TAutoConsoleVariable<bool> UPixelStreaming2PluginSettings::CVarEditorStartOnLaun
 	TEXT("Start Editor Streaming as soon as the Unreal Editor is launched. Default: false"),
 	ECVF_Default);
 
+TAutoConsoleVariable<bool> UPixelStreaming2PluginSettings::CVarAutoStreamPIE(
+	TEXT("PixelStreaming2.Editor.AutoStreamPIE"),
+	true,
+	TEXT("Automatically start streaming when playing in the editor. Default: true"),
+	ECVF_Default);
+
 TAutoConsoleVariable<bool> UPixelStreaming2PluginSettings::CVarEditorUseRemoteSignallingServer(
 	TEXT("PixelStreaming2.Editor.UseRemoteSignallingServer"),
 	false,
