[PS2] Fix: Crashes due to unsafe usages of FAudioDeviceManager::Get()

diff --git a/Engine/Plugins/Media/PixelStreaming2/Source/PixelStreaming2/Private/AudioCapturer.cpp b/Engine/Plugins/Media/PixelStreaming2/Source/PixelStreaming2/Private/AudioCapturer.cpp
--- a/Engine/Plugins/Media/PixelStreaming2/Source/PixelStreaming2/Private/AudioCapturer.cpp
+++ b/Engine/Plugins/Media/PixelStreaming2/Source/PixelStreaming2/Private/AudioCapturer.cpp
@@ -207,9 +207,12 @@ namespace UE::PixelStreaming2
 		}
 		Mixer = MakeShared<FAudioPatchMixer>(NumChannels, SampleRate, SampleSizeSeconds);
 		MixerTask = FPixelStreamingTickableTask::Create<FMixAudioTask>(this, Mixer);
-		FAudioDeviceManager::Get()->IterateOverAllDevices([this](Audio::FDeviceId AudioDeviceId, FAudioDevice*) {
-			CreateAudioProducer(AudioDeviceId);
-		});
+		if (FAudioDeviceManager* AudioDeviceManager = FAudioDeviceManager::Get(); AudioDeviceManager != nullptr)
+		{
+			AudioDeviceManager->IterateOverAllDevices([this](Audio::FDeviceId AudioDeviceId, FAudioDevice*) {
+				CreateAudioProducer(AudioDeviceId);
+			});
+		}
 	}
 
 	// Mixed audio input will push its audio to an FPatchInputProxy for mixing
diff --git a/Engine/Plugins/Media/PixelStreaming2/Source/PixelStreaming2/Private/AudioProducer.cpp b/Engine/Plugins/Media/PixelStreaming2/Source/PixelStreaming2/Private/AudioProducer.cpp
--- a/Engine/Plugins/Media/PixelStreaming2/Source/PixelStreaming2/Private/AudioProducer.cpp
+++ b/Engine/Plugins/Media/PixelStreaming2/Source/PixelStreaming2/Private/AudioProducer.cpp
@@ -11,9 +11,12 @@ namespace UE::PixelStreaming2
 	TSharedPtr<FAudioProducer> FAudioProducer::Create(Audio::FDeviceId InAudioDeviceId, TSharedPtr<FPatchInputProxy> InPatchInput)
 	{
 		TSharedPtr<FAudioProducer> Listener = TSharedPtr<FAudioProducer>(new FAudioProducer(InPatchInput));
-		if (FAudioDevice* AudioDevice = FAudioDeviceManager::Get()->GetAudioDeviceRaw(InAudioDeviceId))
+		if (FAudioDeviceManager* AudioDeviceManager = FAudioDeviceManager::Get(); AudioDeviceManager != nullptr)
 		{
-			AudioDevice->RegisterSubmixBufferListener(Listener.ToSharedRef(), AudioDevice->GetMainSubmixObject());
+			if (FAudioDevice* AudioDevice = AudioDeviceManager->GetAudioDeviceRaw(InAudioDeviceId); AudioDevice != nullptr)
+			{
+				AudioDevice->RegisterSubmixBufferListener(Listener.ToSharedRef(), AudioDevice->GetMainSubmixObject());
+			}
 		}
 
 		return Listener;
