[PS2] Fix: Toolbar closing on interaction

diff --git a/Engine/Plugins/Media/PixelStreaming2/Source/PixelStreaming2Editor/Private/PixelStreaming2Toolbar.cpp b/Engine/Plugins/Media/PixelStreaming2/Source/PixelStreaming2Editor/Private/PixelStreaming2Toolbar.cpp
--- a/Engine/Plugins/Media/PixelStreaming2/Source/PixelStreaming2Editor/Private/PixelStreaming2Toolbar.cpp
+++ b/Engine/Plugins/Media/PixelStreaming2/Source/PixelStreaming2Editor/Private/PixelStreaming2Toolbar.cpp
@@ -76,9 +76,12 @@ FString GetSSLPrivateKeyPath()
 
 namespace UE::EditorPixelStreaming2
 {
-	TSharedRef<SWidget> CreateTextboxWithFileSelector(const FString& Label, float MaxWidth, TFunction<FString()> GetTextLambda, TFunction<void(const FText&)> OnTextChangedLambda, TFunction<void(const FText&, const ETextCommit::Type)> OnTextCommittedLambda, TFunction<void()> OnClickedLambda)
+	TSharedRef<SWidget> CreateTextboxWithFileSelector(const FString& Label, float MaxWidth, TFunction<FString()> GetTextLambda, TFunction<void(const FText&)> OnTextChangedLambda, TFunction<void(const FText&, const ETextCommit::Type)> OnTextCommittedLambda, TFunction<void()> OnClickedLambda, TFunction<EVisibility()> GetVisibilityLambda, TFunction<bool()> IsEnabledLambda)
 	{
 		TSharedRef<SWidget> Widget = SNew(SHorizontalBox)
+			.Visibility_Lambda([GetVisibilityLambda]() {
+				return GetVisibilityLambda();
+			})
 			+ SHorizontalBox::Slot()
 				  .AutoWidth()
 				  .VAlign(VAlign_Center)
@@ -100,7 +103,11 @@ namespace UE::EditorPixelStreaming2
 							  })
 							  .OnTextCommitted_Lambda([OnTextCommittedLambda](const FText& Text, const ETextCommit::Type CommitType) {
 								  OnTextCommittedLambda(Text, CommitType);
-							  })]
+							  })
+							  .IsEnabled_Lambda([IsEnabledLambda]() {
+								  return IsEnabledLambda();
+							  })
+							]
 			+ SHorizontalBox::Slot()
 				  .AutoWidth()
 				  .HAlign(HAlign_Right)
@@ -109,6 +116,9 @@ namespace UE::EditorPixelStreaming2
 								  OnClickedLambda();
 								  return FReply::Handled();
 							  })
+							  .IsEnabled_Lambda([IsEnabledLambda]() {
+								  return IsEnabledLambda();
+							  })
 							  .ButtonStyle(FAppStyle::Get(), "SimpleButton")
 								  [SNew(SImage)
 										  .Image(FAppStyle::Get().GetBrush("Icons.BrowseContent"))
@@ -220,12 +230,7 @@ namespace UE::EditorPixelStreaming2
 				SetCodec(EVideoCodec::VP8);
 			}),
 			FCanExecuteAction::CreateLambda([] {
-				bool bCanChangeCodec = true;
-				IPixelStreaming2Module::Get().ForEachStreamer([&bCanChangeCodec](TSharedPtr<IPixelStreaming2Streamer> Streamer) {
-					bCanChangeCodec &= !Streamer->IsStreaming();
-				});
-
-				return bCanChangeCodec;
+				return true;
 			}),
 			FIsActionChecked::CreateLambda([]() {
 				return GetCodec() == EVideoCodec::VP8;
@@ -237,12 +242,7 @@ namespace UE::EditorPixelStreaming2
 				SetCodec(EVideoCodec::VP9);
 			}),
 			FCanExecuteAction::CreateLambda([] {
-				bool bCanChangeCodec = true;
-				IPixelStreaming2Module::Get().ForEachStreamer([&bCanChangeCodec](TSharedPtr<IPixelStreaming2Streamer> Streamer) {
-					bCanChangeCodec &= !Streamer->IsStreaming();
-				});
-
-				return bCanChangeCodec;
+				return true;
 			}),
 			FIsActionChecked::CreateLambda([]() {
 				return GetCodec() == EVideoCodec::VP9;
@@ -255,10 +255,6 @@ namespace UE::EditorPixelStreaming2
 			}),
 			FCanExecuteAction::CreateLambda([] {
 				bool bCanChangeCodec = UE::PixelStreaming2::IsEncoderSupported<FVideoEncoderConfigH264>();
-				IPixelStreaming2Module::Get().ForEachStreamer([&bCanChangeCodec](TSharedPtr<IPixelStreaming2Streamer> Streamer) {
-					bCanChangeCodec &= !Streamer->IsStreaming();
-				});
-
 				return bCanChangeCodec;
 			}),
 			FIsActionChecked::CreateLambda([]() {
@@ -272,10 +268,6 @@ namespace UE::EditorPixelStreaming2
 			}),
 			FCanExecuteAction::CreateLambda([] {
 				bool bCanChangeCodec = UE::PixelStreaming2::IsEncoderSupported<FVideoEncoderConfigAV1>();
-				IPixelStreaming2Module::Get().ForEachStreamer([&bCanChangeCodec](TSharedPtr<IPixelStreaming2Streamer> Streamer) {
-					bCanChangeCodec &= !Streamer->IsStreaming();
-				});
-
 				return bCanChangeCodec;
 			}),
 			FIsActionChecked::CreateLambda([]() {
@@ -305,30 +297,21 @@ namespace UE::EditorPixelStreaming2
 						FUIAction(),
 						FOnGetContent::CreateLambda(
 							[&]() {
-								FMenuBuilder MenuBuilder(true, PluginCommands);
+								FMenuBuilder MenuBuilder(false, PluginCommands, MakeShared<FExtender>(), false, &FCoreStyle::Get(), false);
 
 								// Use external signalling server option
 								MenuBuilder.BeginSection("Signalling Server Location", LOCTEXT("PixelStreaming2SSLocation", "Signalling Server Location"));
 								MenuBuilder.AddMenuEntry(FPixelStreaming2Commands::Get().ExternalSignalling);
 								MenuBuilder.EndSection();
 
-								if (!GetUseRemoteSignallingServer())
-								{
-									// Embedded Signalling Server Config (streamer port & http port)
-									RegisterEmbeddedSignallingServerConfig(MenuBuilder);
-
-									// Signalling Server Viewer URLs
-									TSharedPtr<UE::PixelStreaming2Servers::IServer> SignallingServer = IPixelStreaming2EditorModule::Get().GetSignallingServer();
-									if (SignallingServer.IsValid() && SignallingServer->HasLaunched())
-									{
-										RegisterSignallingServerURLs(MenuBuilder);
-									}
-								}
-								else
-								{
-									// Remote Signalling Server Config (URL)
-									RegisterRemoteSignallingServerConfig(MenuBuilder);
-								}
+								// Embedded Signalling Server Config (streamer port & http port)
+								RegisterEmbeddedSignallingServerConfig(MenuBuilder);
+
+								// Signalling Server Viewer URLs
+								RegisterSignallingServerURLs(MenuBuilder);
+
+								// Remote Signalling Server Config (URL)
+								RegisterRemoteSignallingServerConfig(MenuBuilder);
 
 								// Pixel Streaming streamer controls
 								RegisterStreamerControls(MenuBuilder);
@@ -353,127 +336,214 @@ namespace UE::EditorPixelStreaming2
 
 	void FPixelStreaming2Toolbar::RegisterEmbeddedSignallingServerConfig(FMenuBuilder& MenuBuilder)
 	{
-		MenuBuilder.BeginSection("Signalling Server Options", LOCTEXT("PixelStreaming2EmbeddedSSOptions", "Embedded Signalling Server Options"));
-
-		TSharedPtr<UE::PixelStreaming2Servers::IServer> SignallingServer = IPixelStreaming2EditorModule::Get().GetSignallingServer();
-		if (!SignallingServer.IsValid() || !SignallingServer->HasLaunched())
-		{
-			MenuBuilder.AddMenuEntry(FPixelStreaming2Commands::Get().ServeHttps);
-			if (GetServeHttps())
-			{
-				TSharedRef<SWidget> SSLCertificateInputBlock = CreateTextboxWithFileSelector(
-					TEXT("SSL Certificate: "),																			   // Label
-					150,																								   // MaxWidth
-					[]() { return GetSSLCertificatePath(); },															   // GetTextLambda
-					[](const FText& Text) { SetSSLCertificatePath(Text.ToString()); },									   // OnTextChangedLambda
-					[](const FText& Text, const ETextCommit::Type CommitType) { SetSSLCertificatePath(Text.ToString()); }, // OnTextCommittedLambda
-					[this]() { OnOpenFileBrowserClicked(EFileType::Certificate); }										   // OnClickedLambda
-				);
-				MenuBuilder.AddWidget(SSLCertificateInputBlock, FText(), true);
-
-				TSharedRef<SWidget> SSLPrivateKeyInputBlock = CreateTextboxWithFileSelector(
-					TEXT("SSL Private Key: "),																			  // Label
-					150,																								  // MaxWidth
-					[]() { return GetSSLPrivateKeyPath(); },															  // GetTextLambda
-					[](const FText& Text) { SetSSLPrivateKeyPath(Text.ToString()); },									  // OnTextChangedLambda
-					[](const FText& Text, const ETextCommit::Type CommitType) { SetSSLPrivateKeyPath(Text.ToString()); }, // OnTextCommittedLambda
-					[this]() { OnOpenFileBrowserClicked(EFileType::PrivateKey); }										  // OnClickedLambda
-				);
-				MenuBuilder.AddWidget(SSLPrivateKeyInputBlock, FText(), true);
-			}
-
-			TSharedRef<SWidget> StreamerPortInputBlock = SNew(SHorizontalBox)
-				+ SHorizontalBox::Slot()
-					  .AutoWidth()
-					  .VAlign(VAlign_Center)
-					  .Padding(FMargin(36.0f, 3.0f, 8.0f, 3.0f))
-						  [SNew(STextBlock)
-								  .Text(FText::FromString(TEXT("Streamer Port: ")))
-								  .ColorAndOpacity(FSlateColor(FLinearColor(1.0f, 1.0f, 1.0f)))]
-				+ SHorizontalBox::Slot()
-					  .AutoWidth()
-						  [SNew(SNumericEntryBox<int32>)
-								  .MinValue(1)
-								  .Value_Lambda([]() {
-									  return IPixelStreaming2EditorModule::Get().GetStreamerPort();
-								  })
-								  .OnValueChanged_Lambda([](int32 InStreamerPort) {
-									  IPixelStreaming2EditorModule::Get().SetStreamerPort(InStreamerPort);
-								  })
-								  .OnValueCommitted_Lambda([](int32 InStreamerPort, ETextCommit::Type InCommitType) {
-									  IPixelStreaming2EditorModule::Get().SetStreamerPort(InStreamerPort);
-								  })];
-			MenuBuilder.AddWidget(StreamerPortInputBlock, FText(), true);
-			TSharedRef<SWidget> ViewerPortInputBlock = SNew(SHorizontalBox)
-				+ SHorizontalBox::Slot()
-					  .AutoWidth()
-					  .VAlign(VAlign_Center)
-					  .Padding(FMargin(36.0f, 3.0f, 8.0f, 3.0f))
-						  [SNew(STextBlock)
-								  .Text(FText::FromString(TEXT("Viewer Port: ")))
-								  .ColorAndOpacity(FSlateColor(FLinearColor(1.0f, 1.0f, 1.0f)))]
-				+ SHorizontalBox::Slot()
-					  .AutoWidth()
-						  [SNew(SNumericEntryBox<int32>)
-								  .MinValue(1)
-								  .Value_Lambda([]() {
-									  return IPixelStreaming2EditorModule::Get().GetViewerPort();
-								  })
-								  .OnValueChanged_Lambda([](int32 InViewerPort) {
-									  IPixelStreaming2EditorModule::Get().SetViewerPort(InViewerPort);
-								  })
-								  .OnValueCommitted_Lambda([](int32 InViewerPort, ETextCommit::Type InCommitType) {
-									  IPixelStreaming2EditorModule::Get().SetViewerPort(InViewerPort);
-								  })];
-			MenuBuilder.AddWidget(ViewerPortInputBlock, FText(), true);
-			MenuBuilder.AddMenuEntry(FPixelStreaming2Commands::Get().StartSignalling);
-		}
-		else
-		{
-			MenuBuilder.AddMenuEntry(FPixelStreaming2Commands::Get().StopSignalling);
-		}
+		MenuBuilder.BeginSection(
+			"Signalling Server Options", 
+			LOCTEXT("PixelStreaming2EmbeddedSSOptions", "Embedded Signalling Server Options"),
+			TAttribute<EVisibility>::Create([]() {
+        		return !GetUseRemoteSignallingServer()
+            		? EVisibility::Visible
+            		: EVisibility::Collapsed;
+    		}));
+
+		MenuBuilder.AddMenuEntry(
+    		FPixelStreaming2Commands::Get().ServeHttps,       // Command
+    		NAME_None,                                              // InExtensionHook
+    		TAttribute<FText>(),                                    // Label override (leave empty to use default)
+    		TAttribute<FText>(),                                    // Tooltip override (leave empty to use default)
+    		FSlateIcon(),                                           // Icon override (leave default)
+    		NAME_None,                                              // Tutorial highlight
+    		TAttribute<EVisibility>::Create([]() {                 // Visibility lambda
+    		    TSharedPtr<UE::PixelStreaming2Servers::IServer> SignallingServer = IPixelStreaming2EditorModule::Get().GetSignallingServer();
+    		    return !GetUseRemoteSignallingServer()
+    		        ? EVisibility::Visible
+    		        : EVisibility::Collapsed;
+    		})
+		);
+
+		TSharedRef<SWidget> SSLCertificateInputBlock = CreateTextboxWithFileSelector(
+			TEXT("SSL Certificate: "),																			   // Label
+			150,																								   // MaxWidth
+			[]() { return GetSSLCertificatePath(); },															   // GetTextLambda
+			[](const FText& Text) { SetSSLCertificatePath(Text.ToString()); },									   // OnTextChangedLambda
+			[](const FText& Text, const ETextCommit::Type CommitType) { SetSSLCertificatePath(Text.ToString()); }, // OnTextCommittedLambda
+			[this]() { OnOpenFileBrowserClicked(EFileType::Certificate); },										   // OnClickedLambda
+			[]() { return !GetUseRemoteSignallingServer() &&GetServeHttps() ? EVisibility::Visible : EVisibility::Collapsed; },					   // GetVisibilityLambda
+			[]() { 
+				TSharedPtr<UE::PixelStreaming2Servers::IServer> SignallingServer = IPixelStreaming2EditorModule::Get().GetSignallingServer();
+				return !(SignallingServer.IsValid() && SignallingServer->HasLaunched()); 
+			}					   // IsEnabledLambda
+		);
+		MenuBuilder.AddWidget(SSLCertificateInputBlock, FText(), true);
+
+		TSharedRef<SWidget> SSLPrivateKeyInputBlock = CreateTextboxWithFileSelector(
+			TEXT("SSL Private Key: "),																			  // Label
+			150,																								  // MaxWidth
+			[]() { return GetSSLPrivateKeyPath(); },															  // GetTextLambda
+			[](const FText& Text) { SetSSLPrivateKeyPath(Text.ToString()); },									  // OnTextChangedLambda
+			[](const FText& Text, const ETextCommit::Type CommitType) { SetSSLPrivateKeyPath(Text.ToString()); }, // OnTextCommittedLambda
+			[this]() { OnOpenFileBrowserClicked(EFileType::PrivateKey); },										  // OnClickedLambda
+			[]() { return !GetUseRemoteSignallingServer() && GetServeHttps() ? EVisibility::Visible : EVisibility::Collapsed; },					   // GetVisibilityLambda
+			[]() { 
+				TSharedPtr<UE::PixelStreaming2Servers::IServer> SignallingServer = IPixelStreaming2EditorModule::Get().GetSignallingServer();
+				return !(SignallingServer.IsValid() && SignallingServer->HasLaunched()); 
+			}					   // IsEnabledLambda
+		);
+		MenuBuilder.AddWidget(SSLPrivateKeyInputBlock, FText(), true);
+
+		TSharedRef<SWidget> StreamerPortInputBlock = SNew(SHorizontalBox)
+			.Visibility_Lambda([]() {
+				return !GetUseRemoteSignallingServer() ? EVisibility::Visible : EVisibility::Collapsed;
+			})
+			+ SHorizontalBox::Slot()
+			.AutoWidth()
+			.VAlign(VAlign_Center)
+			.Padding(FMargin(36.0f, 3.0f, 8.0f, 3.0f))
+			[
+				SNew(STextBlock)
+					.Text(FText::FromString(TEXT("Streamer Port: ")))
+					.ToolTipText(LOCTEXT("PixelStreaming2StreamerPortTooltip", "The port the signalling server will listen for streamer connections on."))
+					.ColorAndOpacity(FSlateColor(FLinearColor(1.0f, 1.0f, 1.0f)))
+			]
+			+ SHorizontalBox::Slot()
+			.AutoWidth()
+			[
+				SNew(SNumericEntryBox<int32>)
+					.MinValue(1)
+					.Value_Lambda([]() {
+						return IPixelStreaming2EditorModule::Get().GetStreamerPort();
+					})
+					.OnValueChanged_Lambda([](int32 InStreamerPort) {
+						IPixelStreaming2EditorModule::Get().SetStreamerPort(InStreamerPort);
+					})
+					.OnValueCommitted_Lambda([](int32 InStreamerPort, ETextCommit::Type InCommitType) {
+						IPixelStreaming2EditorModule::Get().SetStreamerPort(InStreamerPort);
+					})
+					.IsEnabled_Lambda([]() {
+						TSharedPtr<UE::PixelStreaming2Servers::IServer> SignallingServer = IPixelStreaming2EditorModule::Get().GetSignallingServer();
+						return !(SignallingServer.IsValid() && SignallingServer->HasLaunched()); 
+					})
+			];
+		MenuBuilder.AddWidget(StreamerPortInputBlock, FText(), true);
+
+		TSharedRef<SWidget> ViewerPortInputBlock = SNew(SHorizontalBox)
+			.Visibility_Lambda([]() {
+				return !GetUseRemoteSignallingServer() ? EVisibility::Visible : EVisibility::Collapsed;
+			})
+			+ SHorizontalBox::Slot()
+			.AutoWidth()
+			.VAlign(VAlign_Center)
+			.Padding(FMargin(36.0f, 3.0f, 8.0f, 3.0f))
+			[
+				SNew(STextBlock)
+					.Text(FText::FromString(TEXT("Viewer Port: ")))
+					.ToolTipText(LOCTEXT("PixelStreaming2ViewerPortTooltip", "The port the signalling server will listen for viewer connections on."))
+					.ColorAndOpacity(FSlateColor(FLinearColor(1.0f, 1.0f, 1.0f)))
+			]
+			+ SHorizontalBox::Slot()
+			.AutoWidth()
+			[
+				SNew(SNumericEntryBox<int32>)
+					.MinValue(1)
+					.Value_Lambda([]() {
+						return IPixelStreaming2EditorModule::Get().GetViewerPort();
+					})
+					.OnValueChanged_Lambda([](int32 InViewerPort) {
+						IPixelStreaming2EditorModule::Get().SetViewerPort(InViewerPort);
+					})
+					.OnValueCommitted_Lambda([](int32 InViewerPort, ETextCommit::Type InCommitType) {
+						IPixelStreaming2EditorModule::Get().SetViewerPort(InViewerPort);
+					})
+					.IsEnabled_Lambda([]() {
+						TSharedPtr<UE::PixelStreaming2Servers::IServer> SignallingServer = IPixelStreaming2EditorModule::Get().GetSignallingServer();
+						return !(SignallingServer.IsValid() && SignallingServer->HasLaunched()); 
+					})
+			];
+		MenuBuilder.AddWidget(ViewerPortInputBlock, FText(), true);
+
+		MenuBuilder.AddMenuEntry(
+    		FPixelStreaming2Commands::Get().StartSignalling,       // Command
+    		NAME_None,                                              // InExtensionHook
+    		TAttribute<FText>(),                                    // Label override (leave empty to use default)
+    		TAttribute<FText>(),                                    // Tooltip override (leave empty to use default)
+    		FSlateIcon(),                                           // Icon override (leave default)
+    		NAME_None,                                              // Tutorial highlight
+    		TAttribute<EVisibility>::Create([]() {                 // Visibility lambda
+    		    TSharedPtr<UE::PixelStreaming2Servers::IServer> SignallingServer = IPixelStreaming2EditorModule::Get().GetSignallingServer();
+    		    return !GetUseRemoteSignallingServer() && !(SignallingServer.IsValid() && SignallingServer->HasLaunched())
+    		        ? EVisibility::Visible
+    		        : EVisibility::Collapsed;
+    		})
+		);
+
+		MenuBuilder.AddMenuEntry(
+    		FPixelStreaming2Commands::Get().StopSignalling,       // Command
+    		NAME_None,                                              // InExtensionHook
+    		TAttribute<FText>(),                                    // Label override (leave empty to use default)
+    		TAttribute<FText>(),                                    // Tooltip override (leave empty to use default)
+    		FSlateIcon(),                                           // Icon override (leave default)
+    		NAME_None,                                              // Tutorial highlight
+    		TAttribute<EVisibility>::Create([]() {                 // Visibility lambda
+    		    TSharedPtr<UE::PixelStreaming2Servers::IServer> SignallingServer = IPixelStreaming2EditorModule::Get().GetSignallingServer();
+    		    return !GetUseRemoteSignallingServer() &&(SignallingServer.IsValid() && SignallingServer->HasLaunched())
+    		        ? EVisibility::Visible
+    		        : EVisibility::Collapsed;
+    		})
+		);
 
 		MenuBuilder.EndSection();
 	}
 
 	void FPixelStreaming2Toolbar::RegisterRemoteSignallingServerConfig(FMenuBuilder& MenuBuilder)
 	{
-		TSharedPtr<UE::PixelStreaming2Servers::IServer> SignallingServer = IPixelStreaming2EditorModule::Get().GetSignallingServer();
-		MenuBuilder.BeginSection("Remote Signalling Server Options", LOCTEXT("PixelStreaming2RemoteSSOptions", "Remote Signalling Server Options"));
+		MenuBuilder.BeginSection(
+			"Remote Signalling Server Options", 
+			LOCTEXT("PixelStreaming2RemoteSSOptions", "Remote Signalling Server Options"),
+			TAttribute<EVisibility>::Create([]() {
+        		return GetUseRemoteSignallingServer()
+            		? EVisibility::Visible
+            		: EVisibility::Collapsed;
+    		}));
+
 		{
 			TSharedRef<SWidget> URLInputBlock = SNew(SHorizontalBox)
+				.Visibility_Lambda([]() {
+					return GetUseRemoteSignallingServer() ? EVisibility::Visible : EVisibility::Collapsed;
+				})
 				+ SHorizontalBox::Slot()
-					  .AutoWidth()
-					  .VAlign(VAlign_Center)
-					  .Padding(FMargin(36.0f, 3.0f, 8.0f, 3.0f))
-						  [SNew(STextBlock)
-								  .Text(FText::FromString(TEXT("Remote Signalling Server URL")))
-								  .ColorAndOpacity(FSlateColor(FLinearColor(1.0f, 1.0f, 1.0f)))]
+				.AutoWidth()
+				.VAlign(VAlign_Center)
+				.Padding(FMargin(36.0f, 3.0f, 8.0f, 3.0f))
+				[
+					SNew(STextBlock)
+						.Text(FText::FromString(TEXT("Remote Signalling Server URL")))
+						.ColorAndOpacity(FSlateColor(FLinearColor(1.0f, 1.0f, 1.0f)))
+				]
 				+ SHorizontalBox::Slot()
-					  .AutoWidth()
-						  [SNew(SEditableTextBox)
-								  .Text_Lambda([]() {
-									   FPixelStreaming2EditorModule& EditorModule = static_cast<FPixelStreaming2EditorModule&>(IPixelStreaming2EditorModule::Get());
-									   TSharedPtr<IPixelStreaming2Streamer> Streamer = IPixelStreaming2Module::Get().FindStreamer(EditorModule.GetEditorStreamerID());
-									   return FText::FromString(Streamer->GetConnectionURL());
-								  })
-								  .OnTextChanged_Lambda([](const FText& InText) {
-									  IPixelStreaming2Module::Get().ForEachStreamer([InText](TSharedPtr<IPixelStreaming2Streamer> Streamer) {
-										  Streamer->SetConnectionURL(InText.ToString());
-									  });
-								  })
-								  .OnTextCommitted_Lambda([](const FText& InText, ETextCommit::Type InTextCommit) {
-									  IPixelStreaming2Module::Get().ForEachStreamer([InText](TSharedPtr<IPixelStreaming2Streamer> Streamer) {
-										  Streamer->SetConnectionURL(InText.ToString());
-									  });
-								  })
-								  .IsEnabled_Lambda([]() {
-									  bool bCanChangeURL = true;
-									  IPixelStreaming2Module::Get().ForEachStreamer([&bCanChangeURL](TSharedPtr<IPixelStreaming2Streamer> Streamer) {
-										  bCanChangeURL &= !Streamer->IsStreaming();
-									  });
-									  return bCanChangeURL;
-								  })];
+				.AutoWidth()
+				[
+					SNew(SEditableTextBox)
+						.Text_Lambda([]() {
+							FPixelStreaming2EditorModule& EditorModule = static_cast<FPixelStreaming2EditorModule&>(IPixelStreaming2EditorModule::Get());
+							TSharedPtr<IPixelStreaming2Streamer> Streamer = IPixelStreaming2Module::Get().FindStreamer(EditorModule.GetEditorStreamerID());
+							return FText::FromString(Streamer->GetConnectionURL());
+						})
+						.OnTextChanged_Lambda([](const FText& InText) {
+							IPixelStreaming2Module::Get().ForEachStreamer([InText](TSharedPtr<IPixelStreaming2Streamer> Streamer) {
+								Streamer->SetConnectionURL(InText.ToString());
+							});
+						})
+						.OnTextCommitted_Lambda([](const FText& InText, ETextCommit::Type InTextCommit) {
+							IPixelStreaming2Module::Get().ForEachStreamer([InText](TSharedPtr<IPixelStreaming2Streamer> Streamer) {
+								Streamer->SetConnectionURL(InText.ToString());
+							});
+						})
+						.IsEnabled_Lambda([]() {
+							bool bCanChangeURL = true;
+							IPixelStreaming2Module::Get().ForEachStreamer([&bCanChangeURL](TSharedPtr<IPixelStreaming2Streamer> Streamer) {
+								bCanChangeURL &= !Streamer->IsStreaming();
+							});
+							return bCanChangeURL;
+						})];
 			MenuBuilder.AddWidget(URLInputBlock, FText(), true);
 		}
 		MenuBuilder.EndSection();
@@ -481,36 +551,78 @@ namespace UE::EditorPixelStreaming2
 
 	void FPixelStreaming2Toolbar::RegisterSignallingServerURLs(FMenuBuilder& MenuBuilder)
 	{
-		MenuBuilder.BeginSection("Signalling Server URLs", LOCTEXT("PixelStreaming2SignallingURLs", "Signalling Server URLs"));
+		TSharedRef<SVerticalBox> AdapterBox = SNew(SVerticalBox);
+		
+		TArray<TSharedPtr<FInternetAddr>> AdapterAddresses;
+		ISocketSubsystem::Get(PLATFORM_SOCKETSUBSYSTEM)->GetLocalAdapterAddresses(AdapterAddresses);
+		for (TSharedPtr<FInternetAddr> AdapterAddress : AdapterAddresses)
 		{
-			MenuBuilder.AddWidget(SNew(SBox)
-									  .Padding(FMargin(16.0f, 3.0f))
-										  [SNew(STextBlock)
-												  .ColorAndOpacity(FSlateColor::UseSubduedForeground())
-												  .Text(LOCTEXT("SignallingTip", "The Signalling Server is running and may be accessed via the following URLs (network settings permitting):"))
-												  .WrapTextAt(400)],
-				FText());
-
-			MenuBuilder.AddWidget(SNew(SBox)
-									  .Padding(FMargin(32.0f, 3.0f))
-										  [SNew(STextBlock)
-												  .ColorAndOpacity(FSlateColor::UseSubduedForeground())
-												  .Text(FText::FromString(FString::Printf(TEXT("%s://127.0.0.1:%d"), GetServeHttps() ? TEXT("https") : TEXT("http"), IPixelStreaming2EditorModule::Get().GetViewerPort())))],
-				FText());
-
-			TArray<TSharedPtr<FInternetAddr>> AdapterAddresses;
-			if (ISocketSubsystem::Get(PLATFORM_SOCKETSUBSYSTEM)->GetLocalAdapterAddresses(AdapterAddresses))
-			{
-				for (TSharedPtr<FInternetAddr> AdapterAddress : AdapterAddresses)
-				{
-					MenuBuilder.AddWidget(SNew(SBox)
-											  .Padding(FMargin(32.0f, 3.0f))
-												  [SNew(STextBlock)
-														  .ColorAndOpacity(FSlateColor::UseSubduedForeground())
-														  .Text(FText::FromString(FString::Printf(TEXT("%s://%s:%d"), GetServeHttps() ? TEXT("https") : TEXT("http"), *AdapterAddress->ToString(false), IPixelStreaming2EditorModule::Get().GetViewerPort())))],
-						FText());
-				}
-			}
+		    AdapterBox->AddSlot()
+		    .AutoHeight()
+		    .Padding(FMargin(32.f, 3.f, 0.f, 3.f))
+		    [
+		        SNew(STextBlock)
+		        	.ColorAndOpacity(FSlateColor::UseSubduedForeground())
+		        	.Text_Lambda([AdapterAddress]() {
+		        	    return FText::FromString(FString::Printf(
+		        	        TEXT("%s://%s:%d"),
+		        	        GetServeHttps() ? TEXT("https") : TEXT("http"),
+		        	        *AdapterAddress->ToString(false),
+		        	        IPixelStreaming2EditorModule::Get().GetViewerPort()
+		        	    ));
+		        	})
+		    ];
+		}
+
+
+		MenuBuilder.BeginSection(
+			"Signalling Server URLs", 
+			LOCTEXT("PixelStreaming2SignallingURLs", "Signalling Server URLs"),
+			TAttribute<EVisibility>::Create([]() {
+        		TSharedPtr<UE::PixelStreaming2Servers::IServer> SignallingServer = IPixelStreaming2EditorModule::Get().GetSignallingServer();
+        		return (SignallingServer.IsValid() && SignallingServer->HasLaunched())
+            		? EVisibility::Visible
+            		: EVisibility::Collapsed;
+    		})
+		);
+		{
+			MenuBuilder.AddWidget(
+				SNew(SVerticalBox)
+					.Visibility_Lambda([]() {
+						TSharedPtr<UE::PixelStreaming2Servers::IServer> SignallingServer = IPixelStreaming2EditorModule::Get().GetSignallingServer();
+						return (SignallingServer.IsValid() && SignallingServer->HasLaunched())
+						? EVisibility::Visible
+						: EVisibility::Collapsed;
+					})
+				+ SVerticalBox::Slot()
+				.AutoHeight()
+				.Padding(FMargin(16.0f, 3.0f))
+				[
+					SNew(STextBlock)
+					.ColorAndOpacity(FSlateColor::UseSubduedForeground())
+					.Text(LOCTEXT("SignallingTip", "The Signalling Server is running and may be accessed via the following URLs (network settings permitting):"))
+					.WrapTextAt(400)
+				]
+				+ SVerticalBox::Slot()
+				.AutoHeight()
+				.Padding(FMargin(32.0f, 3.0f))
+				[
+					SNew(STextBlock)
+					.ColorAndOpacity(FSlateColor::UseSubduedForeground())
+					.Text_Lambda([]() {
+						return FText::FromString(
+							FString::Printf(TEXT("%s://127.0.0.1:%d"),
+								GetServeHttps() ? TEXT("https") : TEXT("http"),
+								IPixelStreaming2EditorModule::Get().GetViewerPort()));
+					})
+				]
+				+ SVerticalBox::Slot()
+				.AutoHeight()
+				[
+					AdapterBox
+				]
+			, FText() // No label needed
+			);
 		}
 		MenuBuilder.EndSection();
 	}
@@ -521,50 +633,94 @@ namespace UE::EditorPixelStreaming2
 			FString StreamerId = Streamer->GetId();
 			MenuBuilder.BeginSection(FName(*StreamerId), FText::FromString(FString::Printf(TEXT("Streamer - %s"), *StreamerId)));
 			{
+				MenuBuilder.AddWidget(
+				    SNew(SVerticalBox)
+				    + SVerticalBox::Slot()
+				    .AutoHeight()
+				    [
+				        SNew(SVerticalBox)
+				        + SVerticalBox::Slot()
+				        .AutoHeight()
+				        [
+				            SNew(SBox)
+				            .Padding(FMargin(16.f, 3.f))
+				            [
+				                SNew(STextBlock)
+				                .ColorAndOpacity(FSlateColor::UseSubduedForeground())
+				                .Text_Lambda([Streamer]() -> FText {
+				                    if (Streamer->IsStreaming())
+				                    {
+				                        FString VideoProducer = TEXT("nothing (no video input)");
+				                        if (TSharedPtr<IPixelStreaming2VideoProducer> Video = Streamer->GetVideoProducer().Pin())
+				                        {
+				                            VideoProducer = Video->ToString();
+				                        }
+				                        return FText::FromString(FString::Printf(TEXT("Streaming %s"), *VideoProducer));
+				                    }
+				                    else
+				                    {
+				                        return FText::GetEmpty();
+				                    }
+				                })
+								.Visibility_Lambda([Streamer]() {
+									return Streamer->IsStreaming() ? EVisibility::Visible : EVisibility::Collapsed;
+								})
+				                .WrapTextAt(400)
+				            ]
+				        ]
+				    ],
+				    FText()
+				);
 
-				if (Streamer->IsStreaming())
+				FPixelStreaming2EditorModule& EditorModule = static_cast<FPixelStreaming2EditorModule&>(IPixelStreaming2EditorModule::Get());
+				if (Streamer->GetId() == EditorModule.GetEditorStreamerID())
 				{
-					FString VideoProducer = TEXT("nothing (no video input)");
-					if (TSharedPtr<IPixelStreaming2VideoProducer> Video = Streamer->GetVideoProducer().Pin())
-					{
-						VideoProducer = Video->ToString();
-					}
-
-					MenuBuilder.AddWidget(SNew(SBox)
-											  .Padding(FMargin(16.0f, 3.0f))
-												  [SNew(STextBlock)
-														  .ColorAndOpacity(FSlateColor::UseSubduedForeground())
-														  .Text(FText::FromString(FString::Printf(TEXT("Streaming %s"), *VideoProducer)))
-														  .WrapTextAt(400)],
-						FText());
+					MenuBuilder.AddMenuEntry(
+    					FPixelStreaming2Commands::Get().StreamLevelEditor,       // Command
+    					NAME_None,                                              // InExtensionHook
+    					TAttribute<FText>(),                                    // Label override (leave empty to use default)
+    					TAttribute<FText>(),                                    // Tooltip override (leave empty to use default)
+    					FSlateIcon(),                                           // Icon override (leave default)
+    					NAME_None,                                              // Tutorial highlight
+    					TAttribute<EVisibility>::Create([Streamer]() { return !Streamer->IsStreaming() ? EVisibility::Visible : EVisibility::Collapsed; })
+					);
 
 					MenuBuilder.AddMenuEntry(
-						LOCTEXT("PixelStreaming2_StopStreaming", "Stop Streaming"),
-						LOCTEXT("PixelStreaming2_StopStreamingToolTip", "Stop this streamer"),
-						FSlateIcon(),
-						FExecuteAction::CreateLambda([Streamer]() {
-							Streamer->StopStreaming();
-						}));
+    					FPixelStreaming2Commands::Get().StreamEditor,       // Command
+    					NAME_None,                                              // InExtensionHook
+    					TAttribute<FText>(),                                    // Label override (leave empty to use default)
+    					TAttribute<FText>(),                                    // Tooltip override (leave empty to use default)
+    					FSlateIcon(),                                           // Icon override (leave default)
+    					NAME_None,                                              // Tutorial highlight
+    					TAttribute<EVisibility>::Create([Streamer]() { return !Streamer->IsStreaming() ? EVisibility::Visible : EVisibility::Collapsed; })
+					);
 				}
 				else
 				{
-					FPixelStreaming2EditorModule& EditorModule = static_cast<FPixelStreaming2EditorModule&>(IPixelStreaming2EditorModule::Get());
-					if (Streamer->GetId() == EditorModule.GetEditorStreamerID())
-					{
-						MenuBuilder.AddMenuEntry(FPixelStreaming2Commands::Get().StreamLevelEditor);
-						MenuBuilder.AddMenuEntry(FPixelStreaming2Commands::Get().StreamEditor);
-					}
-					else
-					{
-						MenuBuilder.AddMenuEntry(
-							LOCTEXT("PixelStreaming2_StartStreaming", "Start Streaming"),
-							LOCTEXT("PixelStreaming2_StartStreamingToolTip", "Start this streamer"),
-							FSlateIcon(),
-							FExecuteAction::CreateLambda([Streamer]() {
-								Streamer->StartStreaming();
-							}));
-					}
+					MenuBuilder.AddMenuEntry(
+						LOCTEXT("PixelStreaming2_StartStreaming", "Start Streaming"),
+						LOCTEXT("PixelStreaming2_StartStreamingToolTip", "Start this streamer"),
+						FSlateIcon(),
+						FExecuteAction::CreateLambda([Streamer]() { Streamer->StartStreaming(); }),
+						NAME_None,
+						EUserInterfaceActionType::Button,
+						NAME_None,
+						TAttribute<FText>(),
+						TAttribute<EVisibility>::Create([Streamer]() { return !Streamer->IsStreaming() ? EVisibility::Visible : EVisibility::Collapsed; })
+					);
 				}
+
+				MenuBuilder.AddMenuEntry(
+						LOCTEXT("PixelStreaming2_StopStreaming", "Stop Streaming"),
+						LOCTEXT("PixelStreaming2_StopStreamingToolTip", "Stop this streamer"),
+						FSlateIcon(),
+						FExecuteAction::CreateLambda([Streamer]() { Streamer->StopStreaming(); }),
+						NAME_None,
+						EUserInterfaceActionType::Button,
+						NAME_None,
+						TAttribute<FText>(),
+						TAttribute<EVisibility>::Create([Streamer]() { return Streamer->IsStreaming() ? EVisibility::Visible : EVisibility::Collapsed; })
+				);
 			}
 			MenuBuilder.EndSection();
 		});
