[PS2] Fix: Incorrect gamma with VPX decode

diff --git a/Engine/Plugins/Media/PixelStreaming2/Source/PixelStreaming2/Private/VideoCapturer.cpp b/Engine/Plugins/Media/PixelStreaming2/Source/PixelStreaming2/Private/VideoCapturer.cpp
--- a/Engine/Plugins/Media/PixelStreaming2/Source/PixelStreaming2/Private/VideoCapturer.cpp
+++ b/Engine/Plugins/Media/PixelStreaming2/Source/PixelStreaming2/Private/VideoCapturer.cpp
@@ -223,7 +223,7 @@ namespace UE::PixelStreaming2
 			switch (FinalFormat)
 			{
 				case PixelCaptureBufferFormat::FORMAT_RHI:
-					return FPixelCaptureCapturerI420ToRHI::Create();
+					return FPixelCaptureCapturerI420ToRHI::Create({ .bIsSRGB = bIsSRGB }); // TODO (william.belcher): This capturer doesn't support scaling as part of the capture process
 				case PixelCaptureBufferFormat::FORMAT_I420:
 					return FPixelCaptureCapturerI420::Create();
 				default:
