[PS2] Fix: -PixelStreamingKeyFilter launch arg not working

diff --git a/Engine/Plugins/Media/PixelStreaming2/Source/PixelStreaming2Input/Private/InputHandler.cpp b/Engine/Plugins/Media/PixelStreaming2/Source/PixelStreaming2Input/Private/InputHandler.cpp
--- a/Engine/Plugins/Media/PixelStreaming2/Source/PixelStreaming2Input/Private/InputHandler.cpp
+++ b/Engine/Plugins/Media/PixelStreaming2/Source/PixelStreaming2Input/Private/InputHandler.cpp
@@ -35,6 +35,9 @@ namespace UE::PixelStreaming2Input
 		, FocusedPos(UnfocusedPos)
 
 	{
+		// Manually call OnFilteredKeysChanged to initially populate the FilteredKeys list
+		OnFilteredKeysChanged(UPixelStreaming2PluginSettings::CVarInputKeyFilter.AsVariable());
+
 		// Register this input handler as an IMotionController. The module handles the registering as an IInputDevice
 		IModularFeatures::Get().RegisterModularFeature(GetModularFeatureName(), this);
 
