[PS2] Fix: PixelStreaming2.StartStreaming and PixelStreaming2.StopStreaming not working

diff --git a/Engine/Plugins/Media/PixelStreaming2/Source/PixelStreaming2/Private/PixelStreaming2Module.cpp b/Engine/Plugins/Media/PixelStreaming2/Source/PixelStreaming2/Private/PixelStreaming2Module.cpp
--- a/Engine/Plugins/Media/PixelStreaming2/Source/PixelStreaming2/Private/PixelStreaming2Module.cpp
+++ b/Engine/Plugins/Media/PixelStreaming2/Source/PixelStreaming2/Private/PixelStreaming2Module.cpp
@@ -57,26 +57,6 @@ namespace UE::PixelStreaming2
 		// Initialize PixelStreaming thread. Handles tasks like audio pushing and conference ticking
 		PixelStreamingThread = MakeShared<FPixelStreamingThread>();
 
-		// Extra Initialisations post loading console commands
-		IConsoleManager::Get().RegisterConsoleCommand(
-			TEXT("PixelStreaming.StartStreaming"),
-			TEXT("Start all streaming sessions"),
-			FConsoleCommandDelegate::CreateLambda([DefaultStreamer = this->DefaultStreamer]() {
-				if (DefaultStreamer)
-				{
-					DefaultStreamer->StartStreaming();
-				}
-			}));
-		IConsoleManager::Get().RegisterConsoleCommand(
-			TEXT("PixelStreaming.StopStreaming"),
-			TEXT("End any existing streaming sessions."),
-			FConsoleCommandDelegate::CreateLambda([DefaultStreamer = this->DefaultStreamer]() {
-				if (DefaultStreamer)
-				{
-					DefaultStreamer->StopStreaming();
-				}
-			}));
-
 		FCoreDelegates::OnAllModuleLoadingPhasesComplete.AddLambda([this]() {
 			if (!IsPlatformSupported())
 			{
@@ -137,6 +117,28 @@ namespace UE::PixelStreaming2
 					DefaultStreamer->StartStreaming();
 				}
 			}
+
+			// Extra Initialisations post loading console commands
+			TWeakPtr<IPixelStreaming2Streamer> WeakDefaultStreamer = DefaultStreamer;
+			IConsoleManager::Get().RegisterConsoleCommand(
+				TEXT("PixelStreaming.StartStreaming"),
+				TEXT("Start all streaming sessions"),
+				FConsoleCommandDelegate::CreateLambda([WeakDefaultStreamer]() {
+					if (TSharedPtr<IPixelStreaming2Streamer> PinnedDefaultStreamer = WeakDefaultStreamer.Pin(); PinnedDefaultStreamer)
+					{
+						PinnedDefaultStreamer->StartStreaming();
+					}
+			}));
+		
+			IConsoleManager::Get().RegisterConsoleCommand(
+				TEXT("PixelStreaming.StopStreaming"),
+				TEXT("End any existing streaming sessions."),
+				FConsoleCommandDelegate::CreateLambda([WeakDefaultStreamer]() {
+					if (TSharedPtr<IPixelStreaming2Streamer> PinnedDefaultStreamer = WeakDefaultStreamer.Pin(); PinnedDefaultStreamer)
+					{
+						PinnedDefaultStreamer->StopStreaming();
+					}
+			}));
 		});
 	}
 
