Updating the infra scripts to point to the new release branch.

diff --git a/Engine/Plugins/Media/PixelStreaming/Resources/WebServers/get_ps_servers.bat b/Engine/Plugins/Media/PixelStreaming/Resources/WebServers/get_ps_servers.bat
--- a/Engine/Plugins/Media/PixelStreaming/Resources/WebServers/get_ps_servers.bat
+++ b/Engine/Plugins/Media/PixelStreaming/Resources/WebServers/get_ps_servers.bat
@@ -96,11 +96,15 @@ if DEFINED UEVersion (
     SET PSInfraTagOrBranch=UE5.6
     SET IsTag=0
   )
+  if "%UEVersion%"=="5.7" (
+    SET PSInfraTagOrBranch=UE5.7
+    SET IsTag=0
+  )
 )
 
 @Rem If no arguments select a specific version, fetch the appropriate default
 if NOT DEFINED PSInfraTagOrBranch (
-    SET PSInfraTagOrBranch=UE5.6
+    SET PSInfraTagOrBranch=UE5.7
     SET IsTag=0
 )
 echo Tag or branch: !PSInfraTagOrBranch!
@@ -208,4 +212,4 @@ echo    /b      Specify a specific branch for the tool to download from repo
 echo    /t      Specify a specific tag for the tool to download from repo
 echo    /r      Specify a specific release url path e.g. [REDACTED_PROFILE_URL]/PixelStreamingInfrastructure/releases/download/${RELEASE_HERE}.zip
 echo    /h      Display this help message
-goto :EOF
\ No newline at end of file
+goto :EOF
diff --git a/Engine/Plugins/Media/PixelStreaming/Resources/WebServers/get_ps_servers.sh b/Engine/Plugins/Media/PixelStreaming/Resources/WebServers/get_ps_servers.sh
--- a/Engine/Plugins/Media/PixelStreaming/Resources/WebServers/get_ps_servers.sh
+++ b/Engine/Plugins/Media/PixelStreaming/Resources/WebServers/get_ps_servers.sh
@@ -26,7 +26,7 @@ print_help() {
 # Set all default variables (e.g. # Name and version of ps-infra that we are downloading)
 PSInfraOrg=EpicGamesExt
 PSInfraRepo=PixelStreamingInfrastructure
-PSInfraTagOrBranch=UE5.6
+PSInfraTagOrBranch=UE5.7
 RefType=heads
 IsTag=0
 ReleaseUrlBase=[REDACTED_PROFILE_URL]/PixelStreamingInfrastructure/releases/download
@@ -96,12 +96,17 @@ then
     PSInfraTagOrBranch=UE5.6
     IsTag=0
   fi
+  if [ "$UEVersion" = "5.7" ]
+  then
+    PSInfraTagOrBranch=UE5.7
+    IsTag=0
+  fi
 fi
 
 # If no arguments select a specific version, fetch the appropriate default
 if [ -z "$PSInfraTagOrBranch" ]
 then
-  PSInfraTagOrBranch=UE5.6
+  PSInfraTagOrBranch=UE5.7
   IsTag=0
 fi
 echo "Tag or branch: $PSInfraTagOrBranch"
@@ -191,4 +196,4 @@ rm ps-infra.tar.gz
 
 #Create a DOWNLOAD_VERSION file, which we use as a comparison file to check if we should auto upgrade when these scripts are run again
 echo "$DownloadVersion" >| DOWNLOAD_VERSION
-exit 0
\ No newline at end of file
+exit 0
diff --git a/Engine/Plugins/Media/PixelStreaming2/Resources/WebServers/get_ps_servers.bat b/Engine/Plugins/Media/PixelStreaming2/Resources/WebServers/get_ps_servers.bat
--- a/Engine/Plugins/Media/PixelStreaming2/Resources/WebServers/get_ps_servers.bat
+++ b/Engine/Plugins/Media/PixelStreaming2/Resources/WebServers/get_ps_servers.bat
@@ -96,11 +96,15 @@ if DEFINED UEVersion (
     SET PSInfraTagOrBranch=UE5.6
     SET IsTag=0
   )
+  if "%UEVersion%"=="5.7" (
+    SET PSInfraTagOrBranch=UE5.7
+    SET IsTag=0
+  )
 )
 
 @Rem If no arguments select a specific version, fetch the appropriate default
 if NOT DEFINED PSInfraTagOrBranch (
-    SET PSInfraTagOrBranch=UE5.6
+    SET PSInfraTagOrBranch=UE5.7
     SET IsTag=0
 )
 echo Tag or branch: !PSInfraTagOrBranch!
@@ -208,4 +212,4 @@ echo    /b      Specify a specific branch for the tool to download from repo
 echo    /t      Specify a specific tag for the tool to download from repo
 echo    /r      Specify a specific release url path e.g. [REDACTED_PROFILE_URL]/PixelStreamingInfrastructure/releases/download/${RELEASE_HERE}.zip
 echo    /h      Display this help message
-goto :EOF
\ No newline at end of file
+goto :EOF
diff --git a/Engine/Plugins/Media/PixelStreaming2/Resources/WebServers/get_ps_servers.sh b/Engine/Plugins/Media/PixelStreaming2/Resources/WebServers/get_ps_servers.sh
--- a/Engine/Plugins/Media/PixelStreaming2/Resources/WebServers/get_ps_servers.sh
+++ b/Engine/Plugins/Media/PixelStreaming2/Resources/WebServers/get_ps_servers.sh
@@ -26,7 +26,7 @@ print_help() {
 # Set all default variables (e.g. # Name and version of ps-infra that we are downloading)
 PSInfraOrg=EpicGamesExt
 PSInfraRepo=PixelStreamingInfrastructure
-PSInfraTagOrBranch=UE5.6
+PSInfraTagOrBranch=UE5.7
 RefType=heads
 IsTag=0
 ReleaseUrlBase=[REDACTED_PROFILE_URL]/PixelStreamingInfrastructure/releases/download
@@ -96,12 +96,17 @@ then
     PSInfraTagOrBranch=UE5.6
     IsTag=0
   fi
+  if [ "$UEVersion" = "5.7" ]
+  then
+    PSInfraTagOrBranch=UE5.7
+    IsTag=0
+  fi
 fi
 
 # If no arguments select a specific version, fetch the appropriate default
 if [ -z "$PSInfraTagOrBranch" ]
 then
-  PSInfraTagOrBranch=UE5.6
+  PSInfraTagOrBranch=UE5.7
   IsTag=0
 fi
 echo "Tag or branch: $PSInfraTagOrBranch"
@@ -191,4 +196,4 @@ rm ps-infra.tar.gz
 
 #Create a DOWNLOAD_VERSION file, which we use as a comparison file to check if we should auto upgrade when these scripts are run again
 echo "$DownloadVersion" >| DOWNLOAD_VERSION
-exit 0
\ No newline at end of file
+exit 0
