[PS2] Fix: Non-functional public API

diff --git a/Engine/Plugins/Media/PixelStreaming2/Source/PixelStreaming2/Internal/AudioCapturer.h b/Engine/Plugins/Media/PixelStreaming2/Source/PixelStreaming2/Internal/AudioCapturer.h
--- a/Engine/Plugins/Media/PixelStreaming2/Source/PixelStreaming2/Internal/AudioCapturer.h
+++ b/Engine/Plugins/Media/PixelStreaming2/Source/PixelStreaming2/Internal/AudioCapturer.h
@@ -29,13 +29,13 @@ namespace UE::PixelStreaming2
 		float  SampleSizeSeconds;
 	};
 
-	class FPatchInputProxy : public IPixelStreaming2AudioProducer
+	class FPatchInputProxy
 	{
 	public:
 		UE_API FPatchInputProxy(TSharedPtr<FAudioPatchMixer> InMixer);
-		UE_API virtual ~FPatchInputProxy() override;
+		UE_API virtual ~FPatchInputProxy();
 
-		UE_API virtual void PushAudio(const float* AudioData, int32 InNumSamples, int32 InNumChannels, int32 InSampleRate) override;
+		UE_API void OnNewAudioFrame(const float* AudioData, int32 InNumSamples, int32 InNumChannels, int32 InSampleRate);
 
 	protected:
 		TSharedPtr<FAudioPatchMixer> Mixer;
@@ -53,7 +53,7 @@ namespace UE::PixelStreaming2
 		virtual ~FMixAudioTask() = default;
 
 		// Begin FPixelStreamingTickableTask
-		UE_API virtual void		   Tick(float DeltaMs) override;
+		UE_API virtual void			  Tick(float DeltaMs) override;
 		UE_API virtual const FString& GetName() const override;
 		// End FPixelStreamingTickableTask
 
@@ -65,18 +65,22 @@ namespace UE::PixelStreaming2
 		TSharedPtr<FAudioPatchMixer> Mixer;
 	};
 
-	class FAudioCapturer : public IPixelStreaming2AudioProducer
+	class FAudioCapturer
 	{
 	public:
 		static UE_API TSharedPtr<FAudioCapturer> Create(const int InSampleRate = 48000, const int InNumChannels = 2, const float InSampleSizeInSeconds = 0.5f);
 		virtual ~FAudioCapturer() = default;
 
 		// Mixed audio input will push its audio to an FPatchInputProxy for mixing
-		UE_API TSharedPtr<FAudioProducer> CreateAudioProducer();
-		UE_API void					   CreateAudioProducer(Audio::FDeviceId AudioDeviceId);
-		UE_API void					   RemoveAudioProducer(Audio::FDeviceId AudioDeviceId);
+		UE_API void						  CreateAudioProducer(Audio::FDeviceId AudioDeviceId);
+		UE_API void						  RemoveAudioProducer(Audio::FDeviceId AudioDeviceId);
 
-		UE_API virtual void PushAudio(const float* AudioData, int32 NumSamples, int32 NumChannels, int32 SampleRate) override;
+		// Methods for managing custom audio producers that will have their audio mixed in with engine audio
+		UE_API void AddAudioProducer(TSharedPtr<IPixelStreaming2AudioProducer> AudioProducer);
+		UE_API void RemoveAudioProducer(TSharedPtr<IPixelStreaming2AudioProducer> AudioProducer);
+
+		// Called when mixed audio has been produced (and optionally dumped) and is ready to be sent for encoding
+		UE_API virtual void PushAudio(const float* AudioData, int32 NumSamples, int32 NumChannels, int32 SampleRate);
 
 		/**
 		 * This is broadcast each time audio is captured. Tracks should bind to this and push the audio into the track
@@ -96,9 +100,12 @@ namespace UE::PixelStreaming2
 
 	protected:
 		TSharedPtr<FAudioPatchMixer> Mixer;
-		TSharedPtr<FMixAudioTask>    MixerTask;
+		TSharedPtr<FMixAudioTask>	 MixerTask;
 
-		TMap<Audio::FDeviceId, TSharedPtr<FAudioProducer>> AudioProducers;
+		// Audio producers created by the engine that capture audio from specific devices
+		TMap<Audio::FDeviceId, TSharedPtr<FAudioProducer>> EngineAudioProducers;
+		// Audio producers created by the user that push audio from arbitrary sources
+		TMap<TSharedPtr<IPixelStreaming2AudioProducer>, TSharedPtr<FAudioProducer>> UserAudioProducers;
 
 		int				  SampleRate;
 		int				  NumChannels;
diff --git a/Engine/Plugins/Media/PixelStreaming2/Source/PixelStreaming2/Internal/AudioProducer.h b/Engine/Plugins/Media/PixelStreaming2/Source/PixelStreaming2/Internal/AudioProducer.h
--- a/Engine/Plugins/Media/PixelStreaming2/Source/PixelStreaming2/Internal/AudioProducer.h
+++ b/Engine/Plugins/Media/PixelStreaming2/Source/PixelStreaming2/Internal/AudioProducer.h
@@ -13,22 +13,22 @@ namespace UE::PixelStreaming2
 	class FPatchInputProxy;
 
 	/**
-	 * An audio input capable of listening to UE submix's as well as receiving user audio via the PushAudio method.
-	 * Any received audio will be passed into the Parent's PushAudio method
+	 * An audio input capable of listening to UE submix's as well as receiving user audio via the OnPushedAudio method.
+	 * Any received audio will be passed to the PatchInput's OnNewAudioFrame method
 	 */
-	class FAudioProducer : public ISubmixBufferListener, public IPixelStreaming2AudioProducer
+	class FAudioProducer : public ISubmixBufferListener
 	{
 	public:
 		static UE_API TSharedPtr<FAudioProducer> Create(Audio::FDeviceId AudioDeviceId, TSharedPtr<FPatchInputProxy> PatchInput);
-		static UE_API TSharedPtr<FAudioProducer> Create(TSharedPtr<FPatchInputProxy> PatchInput);
+		static UE_API TSharedPtr<FAudioProducer> Create(TSharedPtr<IPixelStreaming2AudioProducer> AudioProducer, TSharedPtr<FPatchInputProxy> PatchInput);
 		virtual ~FAudioProducer() = default;
 
-		// For users to manually push non-submix audio into
-		UE_API virtual void PushAudio(const float* AudioData, int32 NumSamples, int32 NumChannels, int32 SampleRate) override;
-
 		// ISubmixBufferListener interface
 		UE_API virtual void OnNewSubmixBuffer(const USoundSubmix* OwningSubmix, float* AudioData, int32 NumSamples, int32 NumChannels, const int32 SampleRate, double AudioClock) override;
 
+		// Listener for audio pushed from custom IPixelStreaming2AudioProducer implementations
+		UE_API void OnPushedAudio(const float* AudioData, int32 NumSamples, int32 NumChannels, int32 SampleRate);
+
 		void ToggleMuted() { bIsMuted = !bIsMuted; }
 
 	protected:
diff --git a/Engine/Plugins/Media/PixelStreaming2/Source/PixelStreaming2/Internal/VideoCapturer.h b/Engine/Plugins/Media/PixelStreaming2/Source/PixelStreaming2/Internal/VideoCapturer.h
--- a/Engine/Plugins/Media/PixelStreaming2/Source/PixelStreaming2/Internal/VideoCapturer.h
+++ b/Engine/Plugins/Media/PixelStreaming2/Source/PixelStreaming2/Internal/VideoCapturer.h
@@ -17,7 +17,7 @@ namespace UE::PixelStreaming2
 	class FVideoCapturer : public IPixelCaptureCapturerSource, public TSharedFromThis<FVideoCapturer>
 	{
 	public:
-		static UE_API TSharedPtr<FVideoCapturer> Create(TSharedPtr<FVideoProducer> VideoProducer = nullptr, bool bIsSRGB = false);
+		static UE_API TSharedPtr<FVideoCapturer> Create(bool bIsSRGB = false);
 		UE_API virtual ~FVideoCapturer();
 
 		// Begin IPixelCaptureCapturerSource Interface
@@ -26,15 +26,14 @@ namespace UE::PixelStreaming2
 
 		bool IsReady() const { return bReady; }
 
-		UE_API void SetVideoProducer(TSharedPtr<FVideoProducer> InVideoProducer);
+		UE_API void SetVideoProducer(TSharedPtr<IPixelStreaming2VideoProducer> InVideoProducer);
 
-		TSharedPtr<FVideoProducer> GetVideoProducer() { return VideoProducer; }
+		TSharedPtr<IPixelStreaming2VideoProducer> GetVideoProducer() { return VideoProducer; }
 
 		UE_API TSharedPtr<IPixelCaptureOutputFrame> WaitForFormat(int32 Format, TOptional<FIntPoint> Resolution = TOptional<FIntPoint>(), uint32 MaxWaitTime = 5000);
 
 		UE_API TSharedPtr<IPixelCaptureOutputFrame> RequestFormat(int32 Format, TOptional<FIntPoint> Resolution = TOptional<FIntPoint>());
 
-
 		UE_API void ResetFrameCapturer();
 
 		UE_API void OnFrame(const IPixelCaptureInputFrame& InputFrame);
@@ -54,7 +53,7 @@ namespace UE::PixelStreaming2
 		FOnFrameCapturerCreated OnFrameCapturerCreated;
 
 	protected:
-		UE_API FVideoCapturer(TSharedPtr<FVideoProducer> VideoProducer, bool bIsSRGB);
+		UE_API FVideoCapturer(bool bIsSRGB);
 
 		int32 LastFrameWidth = -1;
 		int32 LastFrameHeight = -1;
@@ -62,7 +61,7 @@ namespace UE::PixelStreaming2
 		bool  bReady = false;
 		bool  bIsSRGB = false;
 
-		TSharedPtr<FVideoProducer>					 VideoProducer;
+		TSharedPtr<IPixelStreaming2VideoProducer>	 VideoProducer;
 		TSharedPtr<FPixelCaptureCapturerMultiFormat> FrameCapturer;
 		FDelegateHandle								 CaptureCompleteHandle;
 		FDelegateHandle								 FramePushedHandle;
diff --git a/Engine/Plugins/Media/PixelStreaming2/Source/PixelStreaming2/Internal/VideoProducer.h b/Engine/Plugins/Media/PixelStreaming2/Source/PixelStreaming2/Internal/VideoProducer.h
--- a/Engine/Plugins/Media/PixelStreaming2/Source/PixelStreaming2/Internal/VideoProducer.h
+++ b/Engine/Plugins/Media/PixelStreaming2/Source/PixelStreaming2/Internal/VideoProducer.h
@@ -11,7 +11,7 @@ namespace UE::PixelStreaming2
 {
 	namespace VideoProducerIdentifiers
 	{
-		PIXELSTREAMING2_API extern const FString FVideoProducer;
+		PIXELSTREAMING2_API extern const FString FVideoProducerBase;
 		PIXELSTREAMING2_API extern const FString FVideoProducerBackBuffer;
 		PIXELSTREAMING2_API extern const FString FVideoProducerMediaCapture;
 		PIXELSTREAMING2_API extern const FString FVideoProducerPIEViewport;
@@ -28,23 +28,21 @@ namespace UE::PixelStreaming2
 		FString ProducerName = TEXT("");
 	};
 
-	class FVideoProducer : public IPixelStreaming2VideoProducer
+	/**
+	 * The base video producer implementation used by all other PS2 video producers as well as in tests.
+	 */
+	class FVideoProducerBase : public IPixelStreaming2VideoProducer
 	{
 	public:
-		static UE_API TSharedPtr<FVideoProducer> Create();
-		virtual ~FVideoProducer() = default;
-
-		DECLARE_EVENT_OneParam(FVideoProducer, FOnFramePushed, const IPixelCaptureInputFrame&);
-		FOnFramePushed OnFramePushed;
-
-		UE_API virtual void PushFrame(const IPixelCaptureInputFrame& InputFrame) override;
+		static UE_API TSharedPtr<FVideoProducerBase> Create();
+		virtual ~FVideoProducerBase() = default;
 
 		UE_API virtual FString ToString() override;
 
-		virtual bool IsFrameAlreadyCopied() { return false; }
+		virtual EVideoProducerCapabilities GetCapabilities() { return EVideoProducerCapabilities::Default; }
 
 	protected:
-		FVideoProducer() = default;
+		FVideoProducerBase() = default;
 
 		UE_API virtual ETextureCreateFlags GetTexCreateFlags();
 	};
diff --git a/Engine/Plugins/Media/PixelStreaming2/Source/PixelStreaming2/Internal/VideoProducerBackBuffer.h b/Engine/Plugins/Media/PixelStreaming2/Source/PixelStreaming2/Internal/VideoProducerBackBuffer.h
--- a/Engine/Plugins/Media/PixelStreaming2/Source/PixelStreaming2/Internal/VideoProducerBackBuffer.h
+++ b/Engine/Plugins/Media/PixelStreaming2/Source/PixelStreaming2/Internal/VideoProducerBackBuffer.h
@@ -12,7 +12,7 @@ namespace UE::PixelStreaming2
 	/**
 	 * Use this if you want to send the contents of the UE backbuffer.
 	 */
-	class FVideoProducerBackBuffer : public FVideoProducer
+	class FVideoProducerBackBuffer : public FVideoProducerBase
 	{
 	public:
 		static UE_API TSharedPtr<FVideoProducerBackBuffer> Create();
diff --git a/Engine/Plugins/Media/PixelStreaming2/Source/PixelStreaming2/Internal/VideoProducerMediaCapture.h b/Engine/Plugins/Media/PixelStreaming2/Source/PixelStreaming2/Internal/VideoProducerMediaCapture.h
--- a/Engine/Plugins/Media/PixelStreaming2/Source/PixelStreaming2/Internal/VideoProducerMediaCapture.h
+++ b/Engine/Plugins/Media/PixelStreaming2/Source/PixelStreaming2/Internal/VideoProducerMediaCapture.h
@@ -31,7 +31,7 @@ namespace UE::PixelStreaming2
 	/**
 	 * Use this if you want to send media capture frames as video input.
 	 */
-	class FVideoProducerMediaCapture : public FVideoProducer, public TSharedFromThis<FVideoProducerMediaCapture>
+	class FVideoProducerMediaCapture : public FVideoProducerBase, public TSharedFromThis<FVideoProducerMediaCapture>
 	{
 	public:
 		/**
@@ -53,7 +53,7 @@ namespace UE::PixelStreaming2
 
 		UE_API virtual FString ToString() override;
 
-		virtual bool IsFrameAlreadyCopied() override { return true; }
+		virtual EVideoProducerCapabilities GetCapabilities() { return EVideoProducerCapabilities::ProducesPreprocessedFrames; }
 
 	protected:
 		UE_API void StartActiveViewportCapture();
diff --git a/Engine/Plugins/Media/PixelStreaming2/Source/PixelStreaming2/Internal/VideoProducerRenderTarget.h b/Engine/Plugins/Media/PixelStreaming2/Source/PixelStreaming2/Internal/VideoProducerRenderTarget.h
--- a/Engine/Plugins/Media/PixelStreaming2/Source/PixelStreaming2/Internal/VideoProducerRenderTarget.h
+++ b/Engine/Plugins/Media/PixelStreaming2/Source/PixelStreaming2/Internal/VideoProducerRenderTarget.h
@@ -13,7 +13,7 @@ namespace UE::PixelStreaming2
 	/**
 	 * Use this if you want to send the contents of a render target.
 	 */
-	class FVideoProducerRenderTarget : public FVideoProducer
+	class FVideoProducerRenderTarget : public FVideoProducerBase
 	{
 	public:
 		static UE_API TSharedPtr<FVideoProducerRenderTarget> Create(UTextureRenderTarget2D* Target);
diff --git a/Engine/Plugins/Media/PixelStreaming2/Source/PixelStreaming2/Internal/VideoProducerViewportBase.h b/Engine/Plugins/Media/PixelStreaming2/Source/PixelStreaming2/Internal/VideoProducerViewportBase.h
--- a/Engine/Plugins/Media/PixelStreaming2/Source/PixelStreaming2/Internal/VideoProducerViewportBase.h
+++ b/Engine/Plugins/Media/PixelStreaming2/Source/PixelStreaming2/Internal/VideoProducerViewportBase.h
@@ -18,12 +18,12 @@ namespace UE::PixelStreaming2
 	 * Implementations MUST set the TargetViewport and TargetWidget and then call CalculateCaptureRegion at least once to configure the captured region.
 	 * Additionally, implementations MUST override ShouldCaptureViewport with an optional check to prevent OnBackBufferReadyToPresent from sending frames.
 	 */
-	class FVideoProducerViewportBase : public FVideoProducer
+	class FVideoProducerViewportBase : public FVideoProducerBase
 	{
 	public:
 		virtual ~FVideoProducerViewportBase() = default;
 
-		UE_API virtual bool IsFrameAlreadyCopied() override;
+		UE_API virtual EVideoProducerCapabilities GetCapabilities() override;
 
     protected:
         FVideoProducerViewportBase() = default;
diff --git a/Engine/Plugins/Media/PixelStreaming2/Source/PixelStreaming2/Private/AudioCapturer.cpp b/Engine/Plugins/Media/PixelStreaming2/Source/PixelStreaming2/Private/AudioCapturer.cpp
--- a/Engine/Plugins/Media/PixelStreaming2/Source/PixelStreaming2/Private/AudioCapturer.cpp
+++ b/Engine/Plugins/Media/PixelStreaming2/Source/PixelStreaming2/Private/AudioCapturer.cpp
@@ -50,7 +50,7 @@ namespace UE::PixelStreaming2
 		Mixer->RemovePatch(PatchInput);
 	}
 
-	void FPatchInputProxy::PushAudio(const float* AudioData, int32 InNumSamples, int32 InNumChannels, int32 InSampleRate)
+	void FPatchInputProxy::OnNewAudioFrame(const float* AudioData, int32 InNumSamples, int32 InNumChannels, int32 InSampleRate)
 	{
 		TArray<float> AudioBuffer;
 
@@ -88,7 +88,7 @@ namespace UE::PixelStreaming2
 				// Up mix by duplicating the mono source to each channel
 				TArrayView<int16> SourceBuffer = Buffer.GetArrayView();
 				TArray<int16>	  MixedBuffer;
-				MixedBuffer.SetNumZeroed(InNumSamples / InNumChannels * NumChannels);
+				MixedBuffer.SetNumZeroed(Buffer.GetNumSamples() / InNumChannels * NumChannels);
 				for (int32 SrcSampleIdx = 0; SrcSampleIdx < Buffer.GetNumSamples(); SrcSampleIdx++)
 				{
 					int32 DestSampleIdx = SrcSampleIdx * NumChannels;
@@ -98,7 +98,7 @@ namespace UE::PixelStreaming2
 					}
 				}
 
-				Buffer.CopyFrom(MixedBuffer, NumChannels, InSampleRate);
+				Buffer.CopyFrom(MixedBuffer, NumChannels, SampleRate);
 			}
 			else
 			{
@@ -215,23 +215,27 @@ namespace UE::PixelStreaming2
 		}
 	}
 
-	// Mixed audio input will push its audio to an FPatchInputProxy for mixing
-	TSharedPtr<FAudioProducer> FAudioCapturer::CreateAudioProducer()
-	{
-		// The lifetimes of audio producers created by the user are the responsibility of the user
-		return FAudioProducer::Create(MakeShared<FPatchInputProxy>(Mixer));
-	}
-
 	void FAudioCapturer::CreateAudioProducer(Audio::FDeviceId AudioDeviceId)
 	{
 		// The lifetimes of audio producers created by the engine are our responsibility
 		TSharedPtr<FAudioProducer> AudioInput = FAudioProducer::Create(AudioDeviceId, MakeShared<FPatchInputProxy>(Mixer));
-		AudioProducers.Add(AudioDeviceId, AudioInput);
+		EngineAudioProducers.Add(AudioDeviceId, AudioInput);
 	}
 
 	void FAudioCapturer::RemoveAudioProducer(Audio::FDeviceId AudioDeviceId)
 	{
-		AudioProducers.Remove(AudioDeviceId);
+		EngineAudioProducers.Remove(AudioDeviceId);
+	}
+
+	void FAudioCapturer::AddAudioProducer(TSharedPtr<IPixelStreaming2AudioProducer> AudioProducer)
+	{
+		TSharedPtr<FAudioProducer> AudioInput = FAudioProducer::Create(AudioProducer, MakeShared<FPatchInputProxy>(Mixer));
+		UserAudioProducers.Add(AudioProducer, AudioInput);
+	}
+
+	void FAudioCapturer::RemoveAudioProducer(TSharedPtr<IPixelStreaming2AudioProducer> AudioProducer)
+	{
+		UserAudioProducers.Remove(AudioProducer);
 	}
 
 	void FAudioCapturer::OnAudio(const float* AudioData, int32 InNumSamples, int32 InNumChannels, int32 InSampleRate)
diff --git a/Engine/Plugins/Media/PixelStreaming2/Source/PixelStreaming2/Private/AudioProducer.cpp b/Engine/Plugins/Media/PixelStreaming2/Source/PixelStreaming2/Private/AudioProducer.cpp
--- a/Engine/Plugins/Media/PixelStreaming2/Source/PixelStreaming2/Private/AudioProducer.cpp
+++ b/Engine/Plugins/Media/PixelStreaming2/Source/PixelStreaming2/Private/AudioProducer.cpp
@@ -22,9 +22,11 @@ namespace UE::PixelStreaming2
 		return Listener;
 	}
 
-	TSharedPtr<FAudioProducer> FAudioProducer::Create(TSharedPtr<FPatchInputProxy> InPatchInput)
+	TSharedPtr<FAudioProducer> FAudioProducer::Create(TSharedPtr<IPixelStreaming2AudioProducer> InAudioProducer, TSharedPtr<FPatchInputProxy> InPatchInput)
 	{
 		TSharedPtr<FAudioProducer> Listener = TSharedPtr<FAudioProducer>(new FAudioProducer(InPatchInput));
+		InAudioProducer->OnAudioPushed.AddSP(Listener.ToSharedRef(), &FAudioProducer::OnPushedAudio);
+
 		return Listener;
 	}
 
@@ -33,23 +35,23 @@ namespace UE::PixelStreaming2
 	{
 	}
 
-	void FAudioProducer::PushAudio(const float* AudioData, int32 NumSamples, int32 NumChannels, int32 SampleRate)
+	void FAudioProducer::OnNewSubmixBuffer(const USoundSubmix* OwningSubmix, float* AudioData, int32 NumSamples, int32 NumChannels, const int32 SampleRate, double AudioClock)
 	{
 		if (bIsMuted)
 		{
 			return;
 		}
 
-		PatchInput->PushAudio(AudioData, NumSamples, NumChannels, SampleRate);
+		PatchInput->OnNewAudioFrame(AudioData, NumSamples, NumChannels, SampleRate);
 	}
 
-	void FAudioProducer::OnNewSubmixBuffer(const USoundSubmix* OwningSubmix, float* AudioData, int32 NumSamples, int32 NumChannels, const int32 SampleRate, double AudioClock)
+	void FAudioProducer::OnPushedAudio(const float* AudioData, int32 NumSamples, int32 NumChannels, int32 SampleRate)
 	{
 		if (bIsMuted)
 		{
 			return;
 		}
 
-		PatchInput->PushAudio(AudioData, NumSamples, NumChannels, SampleRate);
+		PatchInput->OnNewAudioFrame(AudioData, NumSamples, NumChannels, SampleRate);
 	}
 } // namespace UE::PixelStreaming2
\ No newline at end of file
diff --git a/Engine/Plugins/Media/PixelStreaming2/Source/PixelStreaming2/Private/PixelStreaming2Module.cpp b/Engine/Plugins/Media/PixelStreaming2/Source/PixelStreaming2/Private/PixelStreaming2Module.cpp
--- a/Engine/Plugins/Media/PixelStreaming2/Source/PixelStreaming2/Private/PixelStreaming2Module.cpp
+++ b/Engine/Plugins/Media/PixelStreaming2/Source/PixelStreaming2/Private/PixelStreaming2Module.cpp
@@ -218,12 +218,11 @@ namespace UE::PixelStreaming2
 
 	TSharedPtr<IPixelStreaming2VideoProducer> FPixelStreaming2Module::CreateVideoProducer()
 	{
-		return FVideoProducer::Create();
+		return nullptr;
 	}
 
 	TSharedPtr<IPixelStreaming2AudioProducer> FPixelStreaming2Module::CreateAudioProducer()
 	{
-		// TODO (william.belcher): This will be fixed with the refactor to audio mixing.
 		return nullptr;
 	}
 
diff --git a/Engine/Plugins/Media/PixelStreaming2/Source/PixelStreaming2/Private/VideoCapturer.cpp b/Engine/Plugins/Media/PixelStreaming2/Source/PixelStreaming2/Private/VideoCapturer.cpp
--- a/Engine/Plugins/Media/PixelStreaming2/Source/PixelStreaming2/Private/VideoCapturer.cpp
+++ b/Engine/Plugins/Media/PixelStreaming2/Source/PixelStreaming2/Private/VideoCapturer.cpp
@@ -18,24 +18,18 @@
 
 namespace UE::PixelStreaming2
 {
-	TSharedPtr<FVideoCapturer> FVideoCapturer::Create(TSharedPtr<FVideoProducer> InVideoProducer, bool bInIsSRGB)
+	TSharedPtr<FVideoCapturer> FVideoCapturer::Create(bool bInIsSRGB)
 	{
-		TSharedPtr<FVideoCapturer> VideoCapturer = TSharedPtr<FVideoCapturer>(new FVideoCapturer(InVideoProducer, bInIsSRGB));
+		TSharedPtr<FVideoCapturer> VideoCapturer = TSharedPtr<FVideoCapturer>(new FVideoCapturer(bInIsSRGB));
 
 		return VideoCapturer;
 	}
 
-	FVideoCapturer::FVideoCapturer(TSharedPtr<FVideoProducer> VideoProducer, bool bIsSRGB)
+	FVideoCapturer::FVideoCapturer(bool bIsSRGB)
 		: bIsSRGB(bIsSRGB)
-		, VideoProducer(VideoProducer)
 	{
 		CreateFrameCapturer();
 
-		if (VideoProducer)
-		{
-			FramePushedHandle = VideoProducer->OnFramePushed.AddRaw(this, &FVideoCapturer::OnFrame);
-		}
-
 		if (UPixelStreaming2PluginSettings::FDelegates* Delegates = UPixelStreaming2PluginSettings::Delegates())
 		{
 			SimulcastEnabledChangedHandle = Delegates->OnSimulcastEnabledChanged.AddRaw(this, &FVideoCapturer::OnSimulcastEnabledChanged);
@@ -59,7 +53,7 @@ namespace UE::PixelStreaming2
 		}
 	}
 
-	void FVideoCapturer::SetVideoProducer(TSharedPtr<FVideoProducer> InVideoProducer)
+	void FVideoCapturer::SetVideoProducer(TSharedPtr<IPixelStreaming2VideoProducer> InVideoProducer)
 	{
 		if (VideoProducer)
 		{
@@ -180,7 +174,9 @@ namespace UE::PixelStreaming2
 			switch (FinalFormat)
 			{
 				case PixelCaptureBufferFormat::FORMAT_RHI:
-					if (VideoProducer && VideoProducer->IsFrameAlreadyCopied() && OutputResolution.X == LastFrameWidth && OutputResolution.Y == LastFrameHeight)
+				{
+					bool bProducesPreprocessedFrames = VideoProducer ? EnumHasAnyFlags(VideoProducer->GetCapabilities(), EVideoProducerCapabilities::ProducesPreprocessedFrames) : false;
+					if (bProducesPreprocessedFrames && OutputResolution.X == LastFrameWidth && OutputResolution.Y == LastFrameHeight)
 					{
 						// If we only have a single layer (and it's scale is 1), we can use the no copy capturer
 						// as we know the output from the media capture will already be the correct format and scale
@@ -204,6 +200,7 @@ namespace UE::PixelStreaming2
 							}
 						}
 					}
+				}
 				case PixelCaptureBufferFormat::FORMAT_I420:
 					if (UPixelStreaming2PluginSettings::CVarUseMediaCapture.GetValueOnAnyThread())
 					{
diff --git a/Engine/Plugins/Media/PixelStreaming2/Source/PixelStreaming2/Private/VideoProducer.cpp b/Engine/Plugins/Media/PixelStreaming2/Source/PixelStreaming2/Private/VideoProducer.cpp
--- a/Engine/Plugins/Media/PixelStreaming2/Source/PixelStreaming2/Private/VideoProducer.cpp
+++ b/Engine/Plugins/Media/PixelStreaming2/Source/PixelStreaming2/Private/VideoProducer.cpp
@@ -6,29 +6,24 @@ namespace UE::PixelStreaming2
 {
 	namespace VideoProducerIdentifiers
 	{
-		const FString FVideoProducer = TEXT("The default video producer - override me");
+		const FString FVideoProducerBase = TEXT("The default video producer - override me");
 		const FString FVideoProducerBackBuffer = TEXT("the Back Buffer");
 		const FString FVideoProducerMediaCapture = TEXT("a Media Capture Video Input");
 		const FString FVideoProducerPIEViewport = TEXT("the PIE Viewport");
 		const FString FVideoProducerRenderTarget = TEXT("a Render Target");
 	}
 
-	TSharedPtr<FVideoProducer> FVideoProducer::Create()
+	TSharedPtr<FVideoProducerBase> FVideoProducerBase::Create()
 	{
-		return TSharedPtr<FVideoProducer>(new FVideoProducer());
+		return TSharedPtr<FVideoProducerBase>(new FVideoProducerBase());
 	}
 
-	void FVideoProducer::PushFrame(const IPixelCaptureInputFrame& InputFrame)
+	FString FVideoProducerBase::ToString()
 	{
-		OnFramePushed.Broadcast(InputFrame);
+		return VideoProducerIdentifiers::FVideoProducerBase;
 	}
 
-	FString FVideoProducer::ToString()
-	{
-		return VideoProducerIdentifiers::FVideoProducer;
-	}
-
-	ETextureCreateFlags FVideoProducer::GetTexCreateFlags()
+	ETextureCreateFlags FVideoProducerBase::GetTexCreateFlags()
 	{
 		ETextureCreateFlags Flags = TexCreate_RenderTargetable | TexCreate_UAV;
 
diff --git a/Engine/Plugins/Media/PixelStreaming2/Source/PixelStreaming2/Private/VideoProducerViewportBase.cpp b/Engine/Plugins/Media/PixelStreaming2/Source/PixelStreaming2/Private/VideoProducerViewportBase.cpp
--- a/Engine/Plugins/Media/PixelStreaming2/Source/PixelStreaming2/Private/VideoProducerViewportBase.cpp
+++ b/Engine/Plugins/Media/PixelStreaming2/Source/PixelStreaming2/Private/VideoProducerViewportBase.cpp
@@ -14,9 +14,9 @@
 
 namespace UE::PixelStreaming2
 {
-	bool FVideoProducerViewportBase::IsFrameAlreadyCopied()
-	{
-		return true;
+	EVideoProducerCapabilities FVideoProducerViewportBase::GetCapabilities() 
+	{ 
+		return EVideoProducerCapabilities::ProducesPreprocessedFrames; 
 	}
 	
 	bool FVideoProducerViewportBase::ShouldCaptureViewport()
diff --git a/Engine/Plugins/Media/PixelStreaming2/Source/PixelStreaming2/Public/IPixelStreaming2Module.h b/Engine/Plugins/Media/PixelStreaming2/Source/PixelStreaming2/Public/IPixelStreaming2Module.h
--- a/Engine/Plugins/Media/PixelStreaming2/Source/PixelStreaming2/Public/IPixelStreaming2Module.h
+++ b/Engine/Plugins/Media/PixelStreaming2/Source/PixelStreaming2/Public/IPixelStreaming2Module.h
@@ -79,6 +79,7 @@ public:
 	 * @brief Creates a new video producer.
 	 * @return An video producer interface for you to push video frames in to.
 	 */
+	UE_DEPRECATED(5.7, "CreateVideoProducer has been deprecated. You can create a class that inherits from IPixelStreaming2VideoProducer and call PushFrame.")
 	virtual TSharedPtr<IPixelStreaming2VideoProducer> CreateVideoProducer() = 0;
 
 	/**
@@ -86,6 +87,7 @@ public:
 	 * @note Users are responsible for the lifetime of this object.
 	 * @return An audio producer interface for you to push audio in to.
 	 */
+	UE_DEPRECATED(5.7, "CreateAudioProducer has been deprecated. You can create a class that inherits from IPixelStreaming2AudioProducer and call PushAudio.")
 	virtual TSharedPtr<IPixelStreaming2AudioProducer> CreateAudioProducer() = 0;
 
 	/**
diff --git a/Engine/Plugins/Media/PixelStreaming2/Source/PixelStreaming2Core/Private/Streamer.cpp b/Engine/Plugins/Media/PixelStreaming2/Source/PixelStreaming2Core/Private/Streamer.cpp
--- a/Engine/Plugins/Media/PixelStreaming2/Source/PixelStreaming2Core/Private/Streamer.cpp
+++ b/Engine/Plugins/Media/PixelStreaming2/Source/PixelStreaming2Core/Private/Streamer.cpp
@@ -21,6 +21,9 @@ namespace UE::PixelStreaming2
 		virtual void										 SetCoupleFramerate(bool bCouple) override {}
 		virtual void										 SetVideoProducer(TSharedPtr<IPixelStreaming2VideoProducer> Input) override {}
 		virtual TWeakPtr<IPixelStreaming2VideoProducer>		 GetVideoProducer() override { return nullptr; }
+		virtual void AddAudioProducer(TSharedPtr<IPixelStreaming2AudioProducer> AudioProducer) override {}
+		virtual void RemoveAudioProducer(TSharedPtr<IPixelStreaming2AudioProducer> AudioProducer) override {}
+		virtual TArray<TWeakPtr<IPixelStreaming2AudioProducer>> GetAudioProducers() override { return {}; }
 		virtual void										 SetConnectionURL(const FString& InSignallingServerURL) override {}
 		virtual FString										 GetConnectionURL() override { return FString(); }
 		virtual FString										 GetId() override { return TEXT("DummyStreamer"); }
diff --git a/Engine/Plugins/Media/PixelStreaming2/Source/PixelStreaming2Core/Public/IPixelStreaming2AudioProducer.h b/Engine/Plugins/Media/PixelStreaming2/Source/PixelStreaming2Core/Public/IPixelStreaming2AudioProducer.h
--- a/Engine/Plugins/Media/PixelStreaming2/Source/PixelStreaming2Core/Public/IPixelStreaming2AudioProducer.h
+++ b/Engine/Plugins/Media/PixelStreaming2/Source/PixelStreaming2Core/Public/IPixelStreaming2AudioProducer.h
@@ -33,5 +33,15 @@ public:
 	 * @param NumChannels Number of audio channels. For example 2 for stero audio.
 	 * @param SampleRate Audio sample rate in samples per second.
 	 */
-	virtual void PushAudio(const float* InBuffer, int32 NumSamples, int32 NumChannels, int32 SampleRate) = 0;
+	virtual void PushAudio(const float* InBuffer, int32 NumSamples, int32 NumChannels, int32 SampleRate) final
+	{
+		OnAudioPushed.Broadcast(InBuffer, NumSamples, NumChannels, SampleRate);
+	}
+
+	/**
+	 * Event triggered when a frame is pushed to the audio producer. You can use this to hook into the frame stream.
+	 * Currently, streamers will listen to this event to know when a new frame is available for encoding.
+	 */
+	DECLARE_EVENT_FourParams(IPixelStreaming2AudioProducer, FOnAudioPushed, const float* /*InBuffer*/, int32 /*NumSamples*/, int32 /*NumChannels*/, int32 /*SampleRate*/);
+	FOnAudioPushed OnAudioPushed;
 };
diff --git a/Engine/Plugins/Media/PixelStreaming2/Source/PixelStreaming2Core/Public/IPixelStreaming2Streamer.h b/Engine/Plugins/Media/PixelStreaming2/Source/PixelStreaming2Core/Public/IPixelStreaming2Streamer.h
--- a/Engine/Plugins/Media/PixelStreaming2/Source/PixelStreaming2Core/Public/IPixelStreaming2Streamer.h
+++ b/Engine/Plugins/Media/PixelStreaming2/Source/PixelStreaming2Core/Public/IPixelStreaming2Streamer.h
@@ -3,6 +3,7 @@
 #pragma once
 
 #include "Features/IModularFeature.h"
+#include "IPixelStreaming2AudioProducer.h"
 #include "IPixelStreaming2AudioSink.h"
 #include "IPixelStreaming2VideoProducer.h"
 #include "IPixelStreaming2VideoSink.h"
@@ -64,6 +65,24 @@ public:
 	 */
 	virtual TWeakPtr<IPixelStreaming2VideoProducer> GetVideoProducer() = 0;
 
+	/**
+	 * @brief Add an audio producer to the streamer.
+	 * @param AudioProducer An IPixelStreaming2AudioProducer that this streamer will mix with engine audio before streaming.
+	 */
+	virtual void AddAudioProducer(TSharedPtr<IPixelStreaming2AudioProducer> AudioProducer) = 0;
+
+	/**
+	 * @brief Remove an audio producer from the streamer.
+	 * @param AudioProducer The IPixelStreaming2AudioProducer to remove.
+	 */
+	virtual void RemoveAudioProducer(TSharedPtr<IPixelStreaming2AudioProducer> AudioProducer) = 0;
+
+	/**
+	 * @brief Get the audio producers currently added to this streamer.
+	 * @return An array of weak pointers to the IPixelStreaming2AudioProducers currently added to this streamer.
+	 */
+	virtual TArray<TWeakPtr<IPixelStreaming2AudioProducer>> GetAudioProducers() = 0;
+
 	/**
 	 * @brief Set the URL this streamer will connect to.
 	 * @param InConnectionURL.
@@ -247,7 +266,7 @@ public:
 	virtual void SetConfigOption(const FName& OptionName, const FString& Value) = 0;
 
 	/**
-	 * @brief Get the Gonfigerartion value for a specific option.
+	 * @brief Get the configuration value for a specific option.
 	 * @param OptionName The name of the config option to get.
 	 * @param Value The value of the config option retrieved.
 	 */
diff --git a/Engine/Plugins/Media/PixelStreaming2/Source/PixelStreaming2Core/Public/IPixelStreaming2VideoProducer.h b/Engine/Plugins/Media/PixelStreaming2/Source/PixelStreaming2Core/Public/IPixelStreaming2VideoProducer.h
--- a/Engine/Plugins/Media/PixelStreaming2/Source/PixelStreaming2Core/Public/IPixelStreaming2VideoProducer.h
+++ b/Engine/Plugins/Media/PixelStreaming2/Source/PixelStreaming2Core/Public/IPixelStreaming2VideoProducer.h
@@ -13,6 +13,13 @@ class UPixelStreaming2VideoProducer : public UInterface
 	GENERATED_BODY()
 };
 
+enum class EVideoProducerCapabilities : uint8
+{
+	Default = 0,
+	ProducesPreprocessedFrames = 1 << 0, // Frames pushed into this producer are preprocessed (e.g. resized, format converted) and ready for encoding.
+};
+ENUM_CLASS_FLAGS(EVideoProducerCapabilities);
+
 /**
  * A "Video Producer" is an object that you use to push video frames into the Pixel Streaming system.
  *
@@ -32,7 +39,23 @@ public:
 	 * Pushes a raw video frame into the Pixel Streaming system.
 	 * @param InputFrame The raw input frame, which may be a GPU texture or a CPU texture based on the underlying type of frame pushed.
 	 */
-	virtual void PushFrame(const IPixelCaptureInputFrame& InputFrame) = 0;
+	virtual void PushFrame(const IPixelCaptureInputFrame& InputFrame) final
+	{
+		OnFramePushed.Broadcast(InputFrame);
+	}
+
+	/**
+	 * Event triggered when a frame is pushed to the video producer. You can use this to hook into the frame stream.
+	 * Currently, streamers will listen to this event to know when a new frame is available for encoding.
+	 */
+	DECLARE_EVENT_OneParam(IPixelStreaming2VideoProducer, FOnFramePushed, const IPixelCaptureInputFrame&);
+	FOnFramePushed OnFramePushed;
+
+	/**
+	 * Returns the capabilities of the video producer.
+	 * @return The capabilities of the video producer.
+	 */
+	virtual EVideoProducerCapabilities GetCapabilities() = 0;
 
     /**
 	 * A human readable identifier used when displaying what the streamer is streaming in the toolbar
diff --git a/Engine/Plugins/Media/PixelStreaming2/Source/PixelStreaming2Editor/Internal/VideoProducerBackBufferComposited.h b/Engine/Plugins/Media/PixelStreaming2/Source/PixelStreaming2Editor/Internal/VideoProducerBackBufferComposited.h
--- a/Engine/Plugins/Media/PixelStreaming2/Source/PixelStreaming2Editor/Internal/VideoProducerBackBufferComposited.h
+++ b/Engine/Plugins/Media/PixelStreaming2/Source/PixelStreaming2Editor/Internal/VideoProducerBackBufferComposited.h
@@ -17,7 +17,7 @@ namespace UE::EditorPixelStreaming2
 	/**
 	 * Use this if you want to send the full UE editor as video input.
 	 */
-	class FVideoProducerBackBufferComposited : public FVideoProducer
+	class FVideoProducerBackBufferComposited : public FVideoProducerBase
 	{
 	public:
 		static UE_API TSharedPtr<FVideoProducerBackBufferComposited> Create();
@@ -25,7 +25,7 @@ namespace UE::EditorPixelStreaming2
 
 		UE_API virtual FString ToString() override;
 
-		virtual bool IsFrameAlreadyCopied() override { return true; }
+		virtual EVideoProducerCapabilities GetCapabilities() { return EVideoProducerCapabilities::ProducesPreprocessedFrames; }
 
 		DECLARE_MULTICAST_DELEGATE_OneParam(FOnFrameSizeChanged, TWeakPtr<FIntRect>);
 		FOnFrameSizeChanged OnFrameSizeChanged;
diff --git a/Engine/Plugins/Media/PixelStreaming2/Source/PixelStreaming2RTC/Internal/EpicRtcVideoCapturer.h b/Engine/Plugins/Media/PixelStreaming2/Source/PixelStreaming2RTC/Internal/EpicRtcVideoCapturer.h
--- a/Engine/Plugins/Media/PixelStreaming2/Source/PixelStreaming2RTC/Internal/EpicRtcVideoCapturer.h
+++ b/Engine/Plugins/Media/PixelStreaming2/Source/PixelStreaming2RTC/Internal/EpicRtcVideoCapturer.h
@@ -13,13 +13,13 @@ namespace UE::PixelStreaming2
 	class FEpicRtcVideoCapturer : public FVideoCapturer
 	{
 	public:
-		static UE_API TSharedPtr<FEpicRtcVideoCapturer> Create(TSharedPtr<FVideoProducer> VideoProducer = nullptr, bool bIsSRGB = false);
+		static UE_API TSharedPtr<FEpicRtcVideoCapturer> Create(bool bIsSRGB = false);
 		virtual ~FEpicRtcVideoCapturer() = default;
 
 		UE_API TRefCountPtr<EpicRtcVideoBufferInterface> GetFrameBuffer();
 
 	private:
-		UE_API FEpicRtcVideoCapturer(TSharedPtr<FVideoProducer> VideoProducer, bool bIsSRGB);
+		UE_API FEpicRtcVideoCapturer(bool bIsSRGB);
 	};
 } // namespace UE::PixelStreaming2
 
diff --git a/Engine/Plugins/Media/PixelStreaming2/Source/PixelStreaming2RTC/Internal/SharedTickableTasks.h b/Engine/Plugins/Media/PixelStreaming2/Source/PixelStreaming2RTC/Internal/SharedTickableTasks.h
--- a/Engine/Plugins/Media/PixelStreaming2/Source/PixelStreaming2RTC/Internal/SharedTickableTasks.h
+++ b/Engine/Plugins/Media/PixelStreaming2/Source/PixelStreaming2RTC/Internal/SharedTickableTasks.h
@@ -2,7 +2,6 @@
 
 #pragma once
 
-#include "EpicRtcAudioCapturer.h"
 #include "EpicRtcConferenceUtils.h"
 #include "TickableTask.h"
 
@@ -15,20 +14,9 @@ namespace UE::PixelStreaming2
 			: TickConferenceTask(InTickConferenceTask)
 		{
 		}
-		
-
-		TSharedPtr<FEpicRtcAudioCapturer> GetAudioCapturer()
-		{
-			if (!AudioCapturer)
-			{
-				AudioCapturer = FEpicRtcAudioCapturer::Create();
-			}
-			return AudioCapturer;
-		}
 
 	private:
 		
 		TSharedPtr<FEpicRtcTickConferenceTask> TickConferenceTask;
-		TSharedPtr<FEpicRtcAudioCapturer>	   AudioCapturer;
 	};
 } // namespace UE::PixelStreaming2
diff --git a/Engine/Plugins/Media/PixelStreaming2/Source/PixelStreaming2RTC/Private/EpicRtcStreamer.cpp b/Engine/Plugins/Media/PixelStreaming2/Source/PixelStreaming2RTC/Private/EpicRtcStreamer.cpp
--- a/Engine/Plugins/Media/PixelStreaming2/Source/PixelStreaming2RTC/Private/EpicRtcStreamer.cpp
+++ b/Engine/Plugins/Media/PixelStreaming2/Source/PixelStreaming2RTC/Private/EpicRtcStreamer.cpp
@@ -186,7 +186,7 @@ namespace UE::PixelStreaming2
 
 	void FEpicRtcStreamer::SetVideoProducer(TSharedPtr<IPixelStreaming2VideoProducer> Producer)
 	{
-		VideoCapturer->SetVideoProducer(StaticCastSharedPtr<FVideoProducer>(Producer));
+		VideoCapturer->SetVideoProducer(Producer);
 	}
 
 	TWeakPtr<IPixelStreaming2VideoProducer> FEpicRtcStreamer::GetVideoProducer()
@@ -194,6 +194,52 @@ namespace UE::PixelStreaming2
 		return VideoCapturer->GetVideoProducer();
 	}
 
+	void FEpicRtcStreamer::AddAudioProducer(TSharedPtr<IPixelStreaming2AudioProducer> AudioProducer) 
+	{
+		{
+			FScopeLock Lock(&CustomAudioProducersCS);
+			CustomAudioProducers.Add(AudioProducer);
+		}
+
+		if (!AudioCapturer)
+		{
+			return;
+		}
+
+		AudioCapturer->AddAudioProducer(AudioProducer);
+	}
+
+	void FEpicRtcStreamer::RemoveAudioProducer(TSharedPtr<IPixelStreaming2AudioProducer> AudioProducer) 
+	{
+		{
+			FScopeLock Lock(&CustomAudioProducersCS);
+			if (CustomAudioProducers.Contains(AudioProducer))
+			{
+				CustomAudioProducers.Remove(AudioProducer);
+			}
+		}
+
+		if (!AudioCapturer)
+		{
+			return;
+		}
+
+		AudioCapturer->RemoveAudioProducer(AudioProducer);
+	}
+
+	TArray<TWeakPtr<IPixelStreaming2AudioProducer>> FEpicRtcStreamer::GetAudioProducers() 
+	{
+		TArray<TWeakPtr<IPixelStreaming2AudioProducer>> WeakArray;
+		WeakArray.SetNum(CustomAudioProducers.Num());
+
+		Algo::Transform(CustomAudioProducers, WeakArray, [](const TSharedPtr<IPixelStreaming2AudioProducer>& Shared)
+		{
+    		return TWeakPtr<IPixelStreaming2AudioProducer>(Shared);
+		});
+
+		return WeakArray;
+	}
+
 	void FEpicRtcStreamer::SetConnectionURL(const FString& InConnectionURL)
 	{
 		CurrentSignallingServerURL = InConnectionURL;
@@ -226,6 +272,18 @@ namespace UE::PixelStreaming2
 		VideoCapturer->ResetFrameCapturer();
 
 		TickableTasks = FPixelStreaming2RTCModule::GetModule()->GetSharedTickableTasks();
+		AudioCapturer = FEpicRtcAudioCapturer::Create();
+
+		{
+			FScopeLock Lock(&CustomAudioProducersCS);
+			if (CustomAudioProducers.Num() > 0)
+			{
+				for (TSharedPtr<IPixelStreaming2AudioProducer>& AudioProducer : CustomAudioProducers)
+				{
+					AudioCapturer->AddAudioProducer(AudioProducer);
+				}
+			}
+		}
 
 		// Broadcast the preconnection event just before we do `TryConnect`
 		OnPreConnection().Broadcast(this);
@@ -279,6 +337,7 @@ namespace UE::PixelStreaming2
 		DeleteAllPlayerSessions();
 
 		TickableTasks.Reset();
+		AudioCapturer.Reset();
 	}
 
 	void FEpicRtcStreamer::OnStatsReady(const FString& PlayerId, const EpicRtcConnectionStats& ConnectionStats)
@@ -1592,7 +1651,7 @@ namespace UE::PixelStreaming2
 		if (ParticipantId == StreamerId)
 		{
 			// NOTE: we don't broadcast the audio track open for a streamer
-			AudioSource = FEpicRtcAudioSource::Create(AudioTrack, TickableTasks->GetAudioCapturer());
+			AudioSource = FEpicRtcAudioSource::Create(AudioTrack, AudioCapturer);
 		}
 		else
 		{
diff --git a/Engine/Plugins/Media/PixelStreaming2/Source/PixelStreaming2RTC/Private/EpicRtcStreamer.h b/Engine/Plugins/Media/PixelStreaming2/Source/PixelStreaming2RTC/Private/EpicRtcStreamer.h
--- a/Engine/Plugins/Media/PixelStreaming2/Source/PixelStreaming2RTC/Private/EpicRtcStreamer.h
+++ b/Engine/Plugins/Media/PixelStreaming2/Source/PixelStreaming2RTC/Private/EpicRtcStreamer.h
@@ -2,6 +2,7 @@
 
 #pragma once
 
+#include "EpicRtcAudioCapturer.h"
 #include "EpicRtcAudioTrackObserver.h"
 #include "EpicRtcAudioTrackObserverFactory.h"
 #include "EpicRtcDataTrackObserver.h"
@@ -52,6 +53,10 @@ namespace UE::PixelStreaming2
 		virtual void									SetVideoProducer(TSharedPtr<IPixelStreaming2VideoProducer> Input) override;
 		virtual TWeakPtr<IPixelStreaming2VideoProducer> GetVideoProducer() override;
 
+		virtual void AddAudioProducer(TSharedPtr<IPixelStreaming2AudioProducer> AudioProducer) override;
+		virtual void RemoveAudioProducer(TSharedPtr<IPixelStreaming2AudioProducer> AudioProducer) override;
+		virtual TArray<TWeakPtr<IPixelStreaming2AudioProducer>> GetAudioProducers() override;
+
 		virtual void	SetConnectionURL(const FString& InConnectionURL) override;
 		virtual FString GetConnectionURL() override;
 
@@ -192,6 +197,7 @@ namespace UE::PixelStreaming2
 		FStreamingStartedEvent StreamingStartedEvent;
 		FStreamingStoppedEvent StreamingStoppedEvent;
 
+		TSharedPtr<FEpicRtcAudioCapturer> AudioCapturer;
 		TSharedPtr<FEpicRtcVideoCapturer> VideoCapturer;
 		TSharedPtr<FVideoSourceGroup>	  VideoSourceGroup;
 		TSharedPtr<FFreezeFrame>		  FreezeFrame;
@@ -206,6 +212,9 @@ namespace UE::PixelStreaming2
 		TSharedPtr<FStreamerReconnectTimer> ReconnectTimer;
 		TSharedPtr<FSharedTickableTasks>	TickableTasks;
 
+		FCriticalSection CustomAudioProducersCS;
+		TArray<TSharedPtr<IPixelStreaming2AudioProducer>> CustomAudioProducers;
+
 	public:
 		// Begin IPixelStreaming2SessionObserver
 		virtual void OnSessionStateUpdate(const EpicRtcSessionState StateUpdate) override;
diff --git a/Engine/Plugins/Media/PixelStreaming2/Source/PixelStreaming2RTC/Private/EpicRtcVideoCapturer.cpp b/Engine/Plugins/Media/PixelStreaming2/Source/PixelStreaming2RTC/Private/EpicRtcVideoCapturer.cpp
--- a/Engine/Plugins/Media/PixelStreaming2/Source/PixelStreaming2RTC/Private/EpicRtcVideoCapturer.cpp
+++ b/Engine/Plugins/Media/PixelStreaming2/Source/PixelStreaming2RTC/Private/EpicRtcVideoCapturer.cpp
@@ -7,13 +7,9 @@
 
 namespace UE::PixelStreaming2
 {
-	TSharedPtr<FEpicRtcVideoCapturer> FEpicRtcVideoCapturer::Create(TSharedPtr<FVideoProducer> InVideoProducer, bool bInIsSRGB)
+	TSharedPtr<FEpicRtcVideoCapturer> FEpicRtcVideoCapturer::Create(bool bInIsSRGB)
 	{
-		TSharedPtr<FEpicRtcVideoCapturer> VideoCapturer = TSharedPtr<FEpicRtcVideoCapturer>(new FEpicRtcVideoCapturer(InVideoProducer, bInIsSRGB));
-		if (InVideoProducer)
-		{
-			VideoCapturer->FramePushedHandle = InVideoProducer->OnFramePushed.AddSP(VideoCapturer.ToSharedRef(), &FEpicRtcVideoCapturer::OnFrame);
-		}
+		TSharedPtr<FEpicRtcVideoCapturer> VideoCapturer = TSharedPtr<FEpicRtcVideoCapturer>(new FEpicRtcVideoCapturer(bInIsSRGB));
 
 		if (UPixelStreaming2PluginSettings::FDelegates* Delegates = UPixelStreaming2PluginSettings::Delegates())
 		{
@@ -24,8 +20,8 @@ namespace UE::PixelStreaming2
 		return VideoCapturer;
 	}
 
-	FEpicRtcVideoCapturer::FEpicRtcVideoCapturer(TSharedPtr<FVideoProducer> VideoProducer, bool bIsSRGB)
-		: FVideoCapturer(VideoProducer, bIsSRGB)
+	FEpicRtcVideoCapturer::FEpicRtcVideoCapturer(bool bIsSRGB)
+		: FVideoCapturer(bIsSRGB)
 	{
 		CreateFrameCapturer();
 	}
diff --git a/Engine/Plugins/Media/PixelStreaming2/Source/PixelStreaming2RTC/Private/EpicRtcVideoSink.cpp b/Engine/Plugins/Media/PixelStreaming2/Source/PixelStreaming2RTC/Private/EpicRtcVideoSink.cpp
--- a/Engine/Plugins/Media/PixelStreaming2/Source/PixelStreaming2RTC/Private/EpicRtcVideoSink.cpp
+++ b/Engine/Plugins/Media/PixelStreaming2/Source/PixelStreaming2RTC/Private/EpicRtcVideoSink.cpp
@@ -28,7 +28,7 @@ namespace UE::PixelStreaming2
 
 	FEpicRtcVideoSink::FEpicRtcVideoSink(TRefCountPtr<EpicRtcVideoTrackInterface> InTrack)
 		: TEpicRtcTrack(InTrack)
-		, VideoCapturer(FVideoCapturer::Create(nullptr, true))
+		, VideoCapturer(FVideoCapturer::Create(true))
 	{
 	}
 
diff --git a/Engine/Plugins/Media/PixelStreaming2/Source/PixelStreaming2RTC/Private/Tests/CodecUtils.cpp b/Engine/Plugins/Media/PixelStreaming2/Source/PixelStreaming2RTC/Private/Tests/CodecUtils.cpp
--- a/Engine/Plugins/Media/PixelStreaming2/Source/PixelStreaming2RTC/Private/Tests/CodecUtils.cpp
+++ b/Engine/Plugins/Media/PixelStreaming2/Source/PixelStreaming2RTC/Private/Tests/CodecUtils.cpp
@@ -21,7 +21,7 @@ namespace UE::PixelStreaming2
 
 		FString								 StreamerName(FString::Printf(TEXT("MockStreamer%d"), StreamerPort));
 		TSharedPtr<IPixelStreaming2Streamer> Streamer = CreateStreamer(StreamerName, StreamerPort);
-		TSharedPtr<FVideoProducer>			 VideoProducer = FVideoProducer::Create();
+		TSharedPtr<FVideoProducerBase>		 VideoProducer = FVideoProducerBase::Create();
 		Streamer->SetVideoProducer(VideoProducer);
 
 		TSharedPtr<FMockPlayer>	   Player = CreatePlayer();
@@ -57,7 +57,7 @@ namespace UE::PixelStreaming2
 
 		FString								 StreamerName(FString::Printf(TEXT("MockStreamer%d"), StreamerPort));
 		TSharedPtr<IPixelStreaming2Streamer> Streamer = CreateStreamer(StreamerName, StreamerPort);
-		TSharedPtr<FVideoProducer>			 VideoProducer = FVideoProducer::Create();
+		TSharedPtr<FVideoProducerBase>		 VideoProducer = FVideoProducerBase::Create();
 		Streamer->SetVideoProducer(VideoProducer);
 
 		TSharedPtr<FMockPlayer>	   Player = CreatePlayer();
@@ -65,7 +65,7 @@ namespace UE::PixelStreaming2
 
 		// Note: Important to couple framerate as we are manually passing frames and don't want any cached frames
 		Streamer->SetCoupleFramerate(true);
-		
+
 		TSharedPtr<bool> bStreamingStarted = MakeShared<bool>(false);
 		Streamer->OnStreamingStarted().AddLambda([bStreamingStarted](IPixelStreaming2Streamer*) {
 			*(bStreamingStarted.Get()) = true;
diff --git a/Engine/Plugins/Media/PixelStreaming2/Source/PixelStreaming2RTC/Private/Tests/DelegateTest.cpp b/Engine/Plugins/Media/PixelStreaming2/Source/PixelStreaming2RTC/Private/Tests/DelegateTest.cpp
--- a/Engine/Plugins/Media/PixelStreaming2/Source/PixelStreaming2RTC/Private/Tests/DelegateTest.cpp
+++ b/Engine/Plugins/Media/PixelStreaming2/Source/PixelStreaming2RTC/Private/Tests/DelegateTest.cpp
@@ -316,7 +316,7 @@ namespace UE::PixelStreaming2
 		const TSharedPtr<UE::PixelStreaming2Servers::IServer> SignallingServer = CreateSignallingServer(StreamerPort, PlayerPort);
 
 		const TSharedPtr<IPixelStreaming2Streamer> Streamer = CreateStreamer(StreamerName, StreamerPort);
-		const TSharedPtr<FVideoProducer>		   VideoProducer = FVideoProducer::Create();
+		const TSharedPtr<FVideoProducerBase>	   VideoProducer = FVideoProducerBase::Create();
 		Streamer->SetVideoProducer(VideoProducer);
 
 		TArray<TSharedPtr<FMockPlayer>> Players;
diff --git a/Engine/Plugins/Media/PixelStreaming2/Source/PixelStreaming2RTC/Private/Tests/MediaDirectionTest.cpp b/Engine/Plugins/Media/PixelStreaming2/Source/PixelStreaming2RTC/Private/Tests/MediaDirectionTest.cpp
--- a/Engine/Plugins/Media/PixelStreaming2/Source/PixelStreaming2RTC/Private/Tests/MediaDirectionTest.cpp
+++ b/Engine/Plugins/Media/PixelStreaming2/Source/PixelStreaming2RTC/Private/Tests/MediaDirectionTest.cpp
@@ -23,7 +23,7 @@ namespace UE::PixelStreaming2
 
 		FString								 StreamerName(FString::Printf(TEXT("MockStreamer%d"), StreamerPort));
 		TSharedPtr<IPixelStreaming2Streamer> Streamer = CreateStreamer(StreamerName, StreamerPort);
-		TSharedPtr<FVideoProducer>			 VideoProducer = FVideoProducer::Create();
+		TSharedPtr<FVideoProducerBase>		 VideoProducer = FVideoProducerBase::Create();
 		Streamer->SetVideoProducer(VideoProducer);
 
 		EMediaDirection PlayerMediaDirection = EMediaDirection::Bidirectional;
diff --git a/Engine/Plugins/Media/PixelStreaming2/Source/PixelStreaming2RTC/Private/Tests/TestUtils.cpp b/Engine/Plugins/Media/PixelStreaming2/Source/PixelStreaming2RTC/Private/Tests/TestUtils.cpp
--- a/Engine/Plugins/Media/PixelStreaming2/Source/PixelStreaming2RTC/Private/Tests/TestUtils.cpp
+++ b/Engine/Plugins/Media/PixelStreaming2/Source/PixelStreaming2RTC/Private/Tests/TestUtils.cpp
@@ -333,7 +333,7 @@ namespace UE::PixelStreaming2
 	TSharedPtr<IPixelStreaming2Streamer> CreateStreamer(const FString& StreamerName, int StreamerPort)
 	{
 		TSharedPtr<IPixelStreaming2Streamer> OutStreamer = IPixelStreaming2Module::Get().CreateStreamer(StreamerName);
-		OutStreamer->SetVideoProducer(FVideoProducer::Create());
+		OutStreamer->SetVideoProducer(FVideoProducerBase::Create());
 		OutStreamer->SetConnectionURL(FString::Printf(TEXT("ws://127.0.0.1:%d"), StreamerPort));
 
 		OutStreamer->OnStreamingStarted().AddLambda([](IPixelStreaming2Streamer*) {
diff --git a/Engine/Plugins/Media/PixelStreaming2/Source/PixelStreaming2RTC/Private/Tests/TestUtils.h b/Engine/Plugins/Media/PixelStreaming2/Source/PixelStreaming2RTC/Private/Tests/TestUtils.h
--- a/Engine/Plugins/Media/PixelStreaming2/Source/PixelStreaming2RTC/Private/Tests/TestUtils.h
+++ b/Engine/Plugins/Media/PixelStreaming2/Source/PixelStreaming2RTC/Private/Tests/TestUtils.h
@@ -55,7 +55,7 @@ namespace UE::PixelStreaming2
 	};
 
 	DEFINE_LATENT_AUTOMATION_COMMAND_ONE_PARAMETER(FWaitSeconds, double, WaitSeconds);
-	DEFINE_LATENT_AUTOMATION_COMMAND_TWO_PARAMETER(FSendSolidColorFrame, TSharedPtr<FVideoProducer>, VideoProducer, FMockVideoFrameConfig, FrameConfig);
+	DEFINE_LATENT_AUTOMATION_COMMAND_TWO_PARAMETER(FSendSolidColorFrame, TSharedPtr<IPixelStreaming2VideoProducer>, VideoProducer, FMockVideoFrameConfig, FrameConfig);
 	DEFINE_LATENT_AUTOMATION_COMMAND_THREE_PARAMETER(FWaitForFrameReceived, double, TimeoutSeconds, TSharedPtr<FMockVideoSink>, VideoSink, FMockVideoFrameConfig, FrameConfig);
 	DEFINE_LATENT_AUTOMATION_COMMAND_THREE_PARAMETER(FCleanupAllPlayers, TSharedPtr<UE::PixelStreaming2Servers::IServer>, OutSignallingServer, TSharedPtr<IPixelStreaming2Streamer>, OutStreamer, TArray<TSharedPtr<FMockPlayer>>, OutPlayers);
 	DEFINE_LATENT_AUTOMATION_COMMAND_TWO_PARAMETER(FWaitForDataChannelOrTimeout, double, TimeoutSeconds, TSharedPtr<FMockPlayer>, OutPlayer);
