[PixelStreaming2] Fix: Deadlocks in PixelStreaming2Thread

diff --git a/Engine/Plugins/Media/PixelStreaming2/Source/PixelStreaming2/Private/PixelStreaming2Module.cpp b/Engine/Plugins/Media/PixelStreaming2/Source/PixelStreaming2/Private/PixelStreaming2Module.cpp
--- a/Engine/Plugins/Media/PixelStreaming2/Source/PixelStreaming2/Private/PixelStreaming2Module.cpp
+++ b/Engine/Plugins/Media/PixelStreaming2/Source/PixelStreaming2/Private/PixelStreaming2Module.cpp
@@ -9,7 +9,7 @@
 #include "PixelStreaming2Common.h"
 #include "PixelStreaming2Delegates.h"
 #include "PixelStreaming2PluginSettings.h"
-#include "Thread.h"
+#include "PixelStreamingThread.h"
 #include "UtilsCoder.h"
 #include "UtilsCore.h"
 #include "VideoProducer.h"
@@ -55,7 +55,7 @@ namespace UE::PixelStreaming2
 		}
 
 		// Initialize PixelStreaming thread. Handles tasks like audio pushing and conference ticking
-		PixelStreamingThread = MakeShared<FPixelStreamingThread>();
+		PixelStreamingThread = GetOrCreatePixelStreamingThread();
 
 		FCoreDelegates::OnAllModuleLoadingPhasesComplete.AddLambda([this]() {
 			if (!IsPlatformSupported())
@@ -302,6 +302,16 @@ namespace UE::PixelStreaming2
 		}
 	}
 
+	TSharedPtr<FPixelStreamingThread> FPixelStreaming2Module::GetOrCreatePixelStreamingThread()
+	{
+		if (!PixelStreamingThread)
+		{
+			PixelStreamingThread = MakeShared<FPixelStreamingThread>();
+		}
+
+		return PixelStreamingThread;
+	}
+
 	void FPixelStreaming2Module::InitDefaultStreamer()
 	{
 		UE_LOGFMT(LogPixelStreaming2, Log, "Initializing default streamer. ID: [{0}], Type: [{1}]", GetDefaultStreamerID(), UPixelStreaming2PluginSettings::CVarDefaultStreamerType.GetValueOnAnyThread());
diff --git a/Engine/Plugins/Media/PixelStreaming2/Source/PixelStreaming2/Private/PixelStreaming2Module.h b/Engine/Plugins/Media/PixelStreaming2/Source/PixelStreaming2/Private/PixelStreaming2Module.h
--- a/Engine/Plugins/Media/PixelStreaming2/Source/PixelStreaming2/Private/PixelStreaming2Module.h
+++ b/Engine/Plugins/Media/PixelStreaming2/Source/PixelStreaming2/Private/PixelStreaming2Module.h
@@ -34,6 +34,8 @@ namespace UE::PixelStreaming2
 		virtual FString									  GetDefaultConnectionURL() override;
 		virtual void									  ForEachStreamer(const TFunction<void(TSharedPtr<IPixelStreaming2Streamer>)>& Func) override;
 
+		TSharedPtr<class FPixelStreamingThread> GetOrCreatePixelStreamingThread();
+
 	private:
 		// Begin IModuleInterface
 		void StartupModule() override;
diff --git a/Engine/Plugins/Media/PixelStreaming2/Source/PixelStreaming2/Private/PixelStreamingThread.cpp b/Engine/Plugins/Media/PixelStreaming2/Source/PixelStreaming2/Private/PixelStreamingThread.cpp
new file mode 100644
--- /dev/null
+++ b/Engine/Plugins/Media/PixelStreaming2/Source/PixelStreaming2/Private/PixelStreamingThread.cpp
@@ -0,0 +1,296 @@
+// Copyright Epic Games, Inc. All Rights Reserved.
+
+#include "PixelStreamingThread.h"
+
+#include "Async/EventCount.h"
+#include "HAL/Event.h"
+#include "HAL/Runnable.h"
+#include "HAL/PlatformProcess.h"
+#include "HAL/PlatformTime.h"
+#include "HAL/RunnableThread.h"
+#include "Logging.h"
+#include "Misc/SingleThreadRunnable.h"
+#include "PixelStreaming2Trace.h"
+#include "TickableTask.h"
+
+namespace UE::PixelStreaming2
+{
+	static TWeakPtr<FPixelStreamingRunnable> PixelStreamingRunnable;
+
+	/**
+	 * The runnable. Handles ticking of all tasks
+	 */
+	class FPixelStreamingRunnable : public FRunnable, public FSingleThreadRunnable
+	{
+	public:
+		bool IsSleeping()
+		{
+			return bIsSleeping;
+		}
+
+		// Begin FRunnable
+		virtual bool Init() override
+		{
+			return true;
+		}
+
+		virtual uint32 Run() override
+		{
+			bIsRunning = true;
+
+			while (bIsRunning)
+			{
+				TRACE_CPUPROFILER_EVENT_SCOPE_ON_CHANNEL_STR("FPixelStreamingRunnable::Run LoopStart", PixelStreaming2Channel);
+				Tick();
+
+				// Sleep 1ms
+				FPlatformProcess::Sleep(0.001f);
+			}
+
+			return 0;
+		}
+
+		virtual void Stop() override
+		{
+			TRACE_CPUPROFILER_EVENT_SCOPE_ON_CHANNEL_STR("FPixelStreamingRunnable::Stop", PixelStreaming2Channel);
+			bIsRunning = false;
+
+			// Wake the thread to ensure it exits
+			TaskEvent.Notify();
+		}
+
+		virtual void Exit() override
+		{
+			TRACE_CPUPROFILER_EVENT_SCOPE_ON_CHANNEL_STR("FPixelStreamingRunnable::Exit", PixelStreaming2Channel);
+			bIsRunning = false;
+
+			// Wake the thread to ensure it exits
+			TaskEvent.Notify();
+		}
+
+		virtual FSingleThreadRunnable* GetSingleThreadInterface() override
+		{
+			TRACE_CPUPROFILER_EVENT_SCOPE_ON_CHANNEL_STR("FPixelStreamingRunnable::GetSingleThreadInterface", PixelStreaming2Channel);
+			return this;
+		}
+		// End FRunnable
+
+		// Begin FSingleThreadRunnable
+		virtual void Tick() override
+		{
+			TRACE_CPUPROFILER_EVENT_SCOPE_ON_CHANNEL_STR("FPixelStreamingRunnable::Tick", PixelStreaming2Channel);
+
+			const uint64 NowCycles = FPlatformTime::Cycles64();
+			const double DeltaMs = FPlatformTime::ToMilliseconds64(NowCycles - LastTickCycles);
+		
+			bool bShouldSleep = false;
+			// Generate a token for the task event. If any tasks are added or removed while ticking
+			// it will prevent the sleep at the end from occuring
+			FEventCountToken Token = TaskEvent.PrepareWait();
+			{
+				// Lock the tasks mutex inside an inner scope so that it is released before we sleep
+				FScopeLock TaskLock(&TasksMutex);
+				{
+					TRACE_CPUPROFILER_EVENT_SCOPE_ON_CHANNEL_STR("FPixelStreamingRunnable::StartTicking", PixelStreaming2Channel);
+					FScopeLock NewTaskLock(&NewTasksMutex);
+					for (auto& NewTask : NewTasks)
+					{
+						Tasks.Add(NewTask);
+					}
+					
+					NewTasks.Empty();
+					
+					bIsTicking = true;
+				}
+				
+				for (auto& Task : Tasks)
+				{
+					// A task may be nulled out due to deletion during our loop. Check for safety
+					TSharedPtr<FPixelStreamingTickableTask> PinnedTask = Task.Pin();
+					if (PinnedTask)
+					{
+						TRACE_CPUPROFILER_EVENT_SCOPE_TEXT_ON_CHANNEL(*PinnedTask->GetName(), PixelStreaming2Channel)
+						PinnedTask->Tick(DeltaMs);
+					}
+				}
+				
+				{
+					TRACE_CPUPROFILER_EVENT_SCOPE_ON_CHANNEL_STR("FPixelStreamingRunnable::FinishTicking", PixelStreaming2Channel);
+					if (bNeedsCleanup)
+					{
+						Tasks.RemoveAll([](const TWeakPtr<FPixelStreamingTickableTask> Entry) { return Entry == nullptr || !Entry.IsValid(); });
+						bNeedsCleanup = false;
+					}
+					
+					// Update the boolean to sleep if there's currently no tasks
+					bShouldSleep = Tasks.Num() == 0;
+
+					bIsTicking = false;
+				}
+			}
+
+			// If there's currently no tasks, try and sleep
+			if (bShouldSleep)
+			{
+				bIsSleeping = true;
+				// If there's currently no tasks, but a task was added or removed during this loop, 
+				// this wait will immediately return and continue
+				TaskEvent.Wait(Token);
+				bIsSleeping = false;
+			}
+
+			LastTickCycles = NowCycles;
+		}
+		// End FSingleThreadRunnable
+
+		FPixelStreamingRunnable()
+			: bIsTicking(false)
+			, bNeedsCleanup(false)
+			, bIsRunning(false)
+			, bIsSleeping(false)
+			, LastTickCycles(FPlatformTime::Cycles64())
+		{
+		}
+
+		virtual ~FPixelStreamingRunnable() = default;
+
+	private:
+		void AddTask(TWeakPtr<FPixelStreamingTickableTask> Task)
+		{
+			TRACE_CPUPROFILER_EVENT_SCOPE_ON_CHANNEL_STR("FPixelStreamingRunnable::AddTask", PixelStreaming2Channel);
+			FScopeLock NewTaskLock(&NewTasksMutex);
+			NewTasks.Add(Task);
+			TaskEvent.Notify();
+		}
+
+		void RemoveTask(FPixelStreamingTickableTask* Task)
+		{
+			TRACE_CPUPROFILER_EVENT_SCOPE_ON_CHANNEL_STR("FPixelStreamingRunnable::RemoveTask", PixelStreaming2Channel);
+			if (Task == nullptr)
+			{
+				return;
+			}
+
+			// Lock TaskLock before NewTaskLock to ensure deadlock does not happen when Tick and StartTicking lock.
+			// Locking matches FTickableObjectBase locking in Tickable.cpp
+			FScopeLock TaskLock(&TasksMutex);
+			FScopeLock NewTaskLock(&NewTasksMutex);
+
+			// Remove from pending list if it hasn't been registered
+			NewTasks.RemoveAll([Task](const TWeakPtr<FPixelStreamingTickableTask> Entry) { 
+				const TSharedPtr<FPixelStreamingTickableTask> Pin = Entry.Pin();
+				return Entry == nullptr || !Pin || Pin.Get() == Task; 
+			});
+
+			// During ticking it is not safe to modify the set so null and mark for later
+			if (bIsTicking)
+			{
+				bNeedsCleanup = true;
+				for (TWeakPtr<FPixelStreamingTickableTask>& LoopTask : Tasks)
+				{
+					if (TSharedPtr<FPixelStreamingTickableTask> Pin = LoopTask.Pin(); Pin && Pin.Get() == Task)
+					{
+						LoopTask = nullptr;
+					}
+				}
+			}
+			else
+			{
+				Tasks.RemoveAll([Task](const TWeakPtr<FPixelStreamingTickableTask> Entry) {
+					TSharedPtr<FPixelStreamingTickableTask> Pin = Entry.Pin();
+					return Entry == nullptr || !Pin || Pin.Get() == Task;
+				});
+			}
+
+			TaskEvent.Notify();
+		}
+
+	private:
+		// Allow the FPixelStreamingTickableTask to access the private add and remove tasks
+		friend FPixelStreamingTickableTask;
+
+		// New tasks that have not yet been added to the Tasks list
+		TArray<TWeakPtr<FPixelStreamingTickableTask>> NewTasks;
+		// Lock for modifying new list 
+		FCriticalSection NewTasksMutex;
+
+		// Tasks to execute every tick
+		TArray<TWeakPtr<FPixelStreamingTickableTask>> Tasks;
+		// This critical section should be locked during entire tick process
+		FCriticalSection TasksMutex;
+
+		// Tasks can removed from any thread so this needs to be thread safe
+		std::atomic<bool> bIsTicking;
+		// Tasks can removed from any thread so this needs to be thread safe
+		std::atomic<bool> bNeedsCleanup;
+		// This thread can be stopped from another thread during shutdown so this needs to be thread safe
+		std::atomic<bool> bIsRunning;
+		//
+		std::atomic<bool> bIsSleeping;
+
+		// Event used to sleep the thread if no tasks exist or wake it up when a task is added/removed or the thread is stopped
+		UE::FEventCount TaskEvent;
+
+		uint64 LastTickCycles;
+	};
+
+	FPixelStreamingThread::FPixelStreamingThread()
+	{
+		TRACE_CPUPROFILER_EVENT_SCOPE_ON_CHANNEL_STR("FPixelStreamingThread Constructor", PixelStreaming2Channel)
+		Runnable = MakeShared<FPixelStreamingRunnable>();
+		PixelStreamingRunnable = Runnable;
+		Thread = TSharedPtr<FRunnableThread>(FRunnableThread::Create(Runnable.Get(), TEXT("Pixel Streaming PixelStreaming Thread")));
+	}
+
+	FPixelStreamingThread::~FPixelStreamingThread()
+	{
+		TRACE_CPUPROFILER_EVENT_SCOPE_ON_CHANNEL_STR("FPixelStreamingThread Destructor", PixelStreaming2Channel);
+		if (Thread)
+		{
+			Thread->Kill();
+			Thread.Reset();
+		}
+
+		if (Runnable)
+		{
+			Runnable->Stop();
+			Runnable.Reset();
+		}
+	}
+
+	bool FPixelStreamingThread::IsSleeping()
+	{
+		if (Runnable)
+		{
+			return Runnable->IsSleeping();
+		}	
+
+		UE_LOGFMT(LogPixelStreaming2, Warning, "Tried to get IsSleeping but runnable is invalid!");
+		return true;
+	}
+
+	/**
+	 * ---------- FPixelStreamingTickableTask ---------------
+	 */
+
+	void FPixelStreamingTickableTask::Register()
+	{
+		TRACE_CPUPROFILER_EVENT_SCOPE_ON_CHANNEL_STR("FPixelStreamingTickableTask::Register", PixelStreaming2Channel);
+		TSharedPtr<FPixelStreamingRunnable> Runnable = PixelStreamingRunnable.Pin();
+		if (Runnable)
+		{
+			Runnable->AddTask(AsWeak());
+		}
+	}
+
+	void FPixelStreamingTickableTask::Unregister()
+	{
+		TRACE_CPUPROFILER_EVENT_SCOPE_ON_CHANNEL_STR("FPixelStreamingTickableTask::Unregister", PixelStreaming2Channel);
+		TSharedPtr<FPixelStreamingRunnable> Runnable = PixelStreamingRunnable.Pin();
+		if (Runnable)
+		{
+			Runnable->RemoveTask(this);
+		}
+	}
+
+} // namespace UE::PixelStreaming2
\ No newline at end of file
diff --git a/Engine/Plugins/Media/PixelStreaming2/Source/PixelStreaming2/Private/PixelStreamingThread.h b/Engine/Plugins/Media/PixelStreaming2/Source/PixelStreaming2/Private/PixelStreamingThread.h
new file mode 100644
--- /dev/null
+++ b/Engine/Plugins/Media/PixelStreaming2/Source/PixelStreaming2/Private/PixelStreamingThread.h
@@ -0,0 +1,24 @@
+// Copyright Epic Games, Inc. All Rights Reserved.
+
+#pragma once
+
+#include "CoreMinimal.h"
+
+namespace UE::PixelStreaming2
+{
+	/**
+	 * The thread. Wraps both the RunnableThread and Runnable into a single point
+	 */
+	class FPixelStreamingThread
+	{
+	public:
+		FPixelStreamingThread();
+		~FPixelStreamingThread();
+
+		bool IsSleeping();
+
+	private:
+		TSharedPtr<class FRunnableThread>		  Thread;
+		TSharedPtr<class FPixelStreamingRunnable> Runnable;
+	};
+} // namespace UE::PixelStreaming2
diff --git a/Engine/Plugins/Media/PixelStreaming2/Source/PixelStreaming2/Private/Tests/PixelStreamingThreadTest.cpp b/Engine/Plugins/Media/PixelStreaming2/Source/PixelStreaming2/Private/Tests/PixelStreamingThreadTest.cpp
new file mode 100644
--- /dev/null
+++ b/Engine/Plugins/Media/PixelStreaming2/Source/PixelStreaming2/Private/Tests/PixelStreamingThreadTest.cpp
@@ -0,0 +1,190 @@
+// Copyright Epic Games, Inc. All Rights Reserved.
+
+#include "PixelStreamingThread.h"
+#include "PixelStreaming2Module.h"
+#include "Logging.h"
+#include "Misc/AutomationTest.h"
+#include "TickableTask.h"
+#include "Tests/AutomationCommon.h"
+
+#if WITH_DEV_AUTOMATION_TESTS
+
+namespace UE::PixelStreaming2
+{
+    class FTestTask : public FPixelStreamingTickableTask
+	{
+	public:
+		FTestTask(int TaskIndex)
+            : Name(FString::Printf(TEXT("Task_%d"), TaskIndex))
+        {
+        }
+
+		virtual ~FTestTask() = default;
+
+		// Begin FPixelStreamingTickableTask
+		virtual void Tick(float DeltaMs) override
+        {
+            TickCount++;
+        }
+
+		virtual const FString& GetName() const override
+        {
+            return Name;
+        }
+		// End FPixelStreamingTickableTask
+
+        uint32 GetTickCount()
+        {
+            return TickCount;
+        }
+
+	private:
+        std::atomic<uint32> TickCount;
+
+        FString Name;
+	};
+
+    class FTestThreadRunnable : public FRunnable
+    {
+    public:
+        FTestThreadRunnable(TFunction<void()> RunFunc)
+            : RunFunc(RunFunc)
+        {
+        }
+
+        virtual ~FTestThreadRunnable() = default;
+
+        virtual uint32 Run() override
+        {
+            RunFunc();
+
+            return 0;
+        }
+
+    private:
+        TFunction<void()> RunFunc;
+    };
+
+    IMPLEMENT_SIMPLE_AUTOMATION_TEST(FPS2ThreadTest, "System.Plugins.PixelStreaming2.FPS2ThreadTest", EAutomationTestFlags::EditorContext | EAutomationTestFlags::ClientContext | EAutomationTestFlags::ProductFilter)
+	bool FPS2ThreadTest::RunTest(const FString& Parameters)
+    {
+        FPixelStreaming2Module* Module = FPixelStreaming2Module::GetModule();
+        if (!Module)
+        {
+            return false;
+        }
+
+        TSharedPtr<FPixelStreamingThread> Thread = Module->GetOrCreatePixelStreamingThread();
+
+        TestTrue(TEXT("Thread should be sleeping with no tasks"), Thread->IsSleeping());
+
+        TSharedPtr<FTestTask> Task = FPixelStreamingTickableTask::Create<FTestTask>(0);
+        FPlatformProcess::Sleep(0.1f); // Sleep to ensure thread has had a chance to wake up
+
+        TestTrue(TEXT("Thread should not be sleeping"), !Thread->IsSleeping());
+        TestTrue(TEXT("Task should have ticked at least once"), Task->GetTickCount() > 0);
+
+        Task = nullptr;
+        FPlatformProcess::Sleep(0.1f); // Sleep to ensure thread has had a chance to process removal and then sleep
+
+        TestTrue(TEXT("Thread should be sleeping with no tasks"), Thread->IsSleeping());
+
+		return true;
+    }
+
+	IMPLEMENT_SIMPLE_AUTOMATION_TEST(FPS2ThreadConcurrencyTest, "System.Plugins.PixelStreaming2.FPS2ThreadConcurrencyTest", EAutomationTestFlags::EditorContext | EAutomationTestFlags::ClientContext | EAutomationTestFlags::ProductFilter)
+	bool FPS2ThreadConcurrencyTest::RunTest(const FString& Parameters)
+	{
+        FPixelStreaming2Module* Module = FPixelStreaming2Module::GetModule();
+        if (!Module)
+        {
+            return false;
+        }
+
+        TSharedPtr<FPixelStreamingThread> Thread = Module->GetOrCreatePixelStreamingThread();
+
+        TestTrue(TEXT("Thread should be sleeping with no tasks"), Thread->IsSleeping());
+
+        constexpr int NumThreads = 4;
+        constexpr int NumIterations = 50;
+
+        TArray<FRunnableThread*> TestThreads;
+        TArray<TSharedPtr<FTestThreadRunnable>> TestRunnables;
+
+        FCriticalSection TasksMutex;
+        TArray<TSharedPtr<FTestTask>> TestTasks;
+
+        std::atomic<int> TaskIndex;
+        TFunction<void()> RunnableFunc = [NumIterations, &TaskIndex, &TasksMutex, &TestTasks]()
+        {
+            for (int32 i = 0; i < NumIterations; ++i)
+            {
+                // Random sleep to simulate unpredictable timing
+                FPlatformProcess::Sleep(FMath::FRandRange(0.001f, 0.01f));
+                // Randomly decide to add or remove
+                if (FMath::RandBool())
+                {
+                    // Add task
+                    auto NewTask = FPixelStreamingTickableTask::Create<FTestTask>(TaskIndex++);
+                    FScopeLock Lock(&TasksMutex);
+                    TestTasks.Add(NewTask);
+                }
+                else
+                {
+                    // Remove a random task if any exist
+                    FScopeLock Lock(&TasksMutex);
+                    if (TestTasks.Num() > 0)
+                    {
+                        int32 Index = FMath::RandRange(0, TestTasks.Num() - 1);
+                        TSharedPtr<FTestTask> ToRemove = TestTasks[Index];
+                        TestTasks.RemoveAt(Index);
+                    }
+                }
+            }
+        };
+
+        // Spawn NumThreads number of threads, each with a runnable that will add/remove tasks at random
+        for (int ThreadIdx = 0; ThreadIdx < NumThreads; ThreadIdx++)
+        {
+            TSharedPtr<FTestThreadRunnable> TestRunnable = MakeShared<FTestThreadRunnable>(RunnableFunc);
+            TestRunnables.Add(TestRunnable);
+            TestThreads.Add(FRunnableThread::Create(TestRunnable.Get(), *FString::Printf(TEXT("TestThread_%d"), ThreadIdx)));
+        }
+
+        // Wait for the spawned threads to complete
+        for (FRunnableThread* T : TestThreads)
+        {
+            T->WaitForCompletion();
+            delete T;
+        }
+
+        // Remove runnable associated with threads
+        TestRunnables.Empty();
+
+        // Sleep to ensure activity on PixelStreamingThread has settled
+        FPlatformProcess::Sleep(0.2f);
+
+        // If there's still tasks valid, then the thread shouldn't be sleeping
+        if (TestTasks.Num() > 0)
+        {
+            TestTrue(TEXT("Thread should not be sleeping"), !Thread->IsSleeping());
+        }
+
+        // Verify that tasks not removed have a valid tick count
+        for (TSharedPtr<FTestTask>& Task : TestTasks)
+        {
+            TestTrue(TEXT("Task should have ticked at least once"), Task->GetTickCount() > 0);
+        }
+
+        // Empty the test tasks (unregistering them in the process)
+        TestTasks.Empty();
+        FPlatformProcess::Sleep(0.1f); // Sleep to ensure thread has had a chance to process removals and then sleep
+
+        // With no tasks valid, the thread should now be sleeping
+        TestTrue(TEXT("Thread should be sleeping with no tasks!"), Thread->IsSleeping());
+
+		return true;
+	}
+} // namespace UE::PixelStreaming2
+
+#endif // WITH_DEV_AUTOMATION_TESTS
diff --git a/Engine/Plugins/Media/PixelStreaming2/Source/PixelStreaming2/Private/Thread.cpp b/Engine/Plugins/Media/PixelStreaming2/Source/PixelStreaming2/Private/Thread.cpp
deleted file mode 100644
diff --git a/Engine/Plugins/Media/PixelStreaming2/Source/PixelStreaming2/Private/Thread.h b/Engine/Plugins/Media/PixelStreaming2/Source/PixelStreaming2/Private/Thread.h
deleted file mode 100644
