[AVCodecs, PixelStreaming2] Fixes:
- Fix: Crash when calling PushAudio
- Fix: Stream state set to connecting when no url is set
- Fix: Input handler not unregistering
- Fix: Crash when getting window title during a MRQ render
- Fix: Crash when PropertyChangedEvent.Property is nullptr
- Fix: Downgrade warning to log for build health

diff --git a/Engine/Plugins/Experimental/AVCodecs/AVCodecsCore/Source/AVCodecsCore/Private/AVResult.cpp b/Engine/Plugins/Experimental/AVCodecs/AVCodecsCore/Source/AVCodecsCore/Private/AVResult.cpp
--- a/Engine/Plugins/Experimental/AVCodecs/AVCodecsCore/Source/AVCodecsCore/Private/AVResult.cpp
+++ b/Engine/Plugins/Experimental/AVCodecs/AVCodecsCore/Source/AVCodecsCore/Private/AVResult.cpp
@@ -126,6 +126,10 @@ void FAVResult::Log() const
 		{
 			UE_LOG(LogAVCodecs, Verbose, TEXT("%s"), *ToString());
 		}
+		else
+		{
+			UE_LOG(LogAVCodecs, Log, TEXT("%s"), *ToString());
+		}
 	}
 	else
 	{
diff --git a/Engine/Plugins/Experimental/AVCodecs/NVCodecs/Source/NVDEC/Private/NVDEC.cpp b/Engine/Plugins/Experimental/AVCodecs/NVCodecs/Source/NVDEC/Private/NVDEC.cpp
--- a/Engine/Plugins/Experimental/AVCodecs/NVCodecs/Source/NVDEC/Private/NVDEC.cpp
+++ b/Engine/Plugins/Experimental/AVCodecs/NVCodecs/Source/NVDEC/Private/NVDEC.cpp
@@ -72,7 +72,7 @@ FNVDEC::FNVDEC()
 		}
 		else
 		{
-			FAVResult::Log(EAVResult::Warning, TEXT("Failed to get NVDEC dll handle. NVDEC module will not be available."), TEXT("NVDEC"));
+			FAVResult::Log(EAVResult::Unknown, TEXT("Failed to get NVDEC dll handle. NVDEC module will not be available."), TEXT("NVDEC"));
 		}
 	}
 	else
diff --git a/Engine/Plugins/Experimental/AVCodecs/NVCodecs/Source/NVENC/Private/NVENC.cpp b/Engine/Plugins/Experimental/AVCodecs/NVCodecs/Source/NVENC/Private/NVENC.cpp
--- a/Engine/Plugins/Experimental/AVCodecs/NVCodecs/Source/NVENC/Private/NVENC.cpp
+++ b/Engine/Plugins/Experimental/AVCodecs/NVCodecs/Source/NVENC/Private/NVENC.cpp
@@ -59,7 +59,7 @@ FNVENC::FNVENC()
 		}
 		else
 		{
-			FAVResult::Log(EAVResult::Warning, TEXT("Failed to get NVENC dll handle. NVENC module will not be available."), TEXT("NVENC"));
+			FAVResult::Log(EAVResult::Unknown, TEXT("Failed to get NVENC dll handle. NVENC module will not be available."), TEXT("NVENC"));
 		}
 	}
 	else
diff --git a/Engine/Plugins/Media/PixelStreaming2/Source/PixelStreaming2/Internal/AudioCapturer.h b/Engine/Plugins/Media/PixelStreaming2/Source/PixelStreaming2/Internal/AudioCapturer.h
--- a/Engine/Plugins/Media/PixelStreaming2/Source/PixelStreaming2/Internal/AudioCapturer.h
+++ b/Engine/Plugins/Media/PixelStreaming2/Source/PixelStreaming2/Internal/AudioCapturer.h
@@ -3,7 +3,10 @@
 #pragma once
 
 #include "DSP/MultithreadedPatching.h"
+#include "AudioDeviceManager.h"
 #include "AudioProducer.h"
+#include "Misc/CoreDelegates.h"
+#include "PixelStreaming2PluginSettings.h"
 #include "SampleBuffer.h"
 #include "TickableTask.h"
 
@@ -17,7 +20,7 @@ namespace UE::PixelStreaming2
 	{
 	public:
 		UE_API FAudioPatchMixer(uint8 NumChannels, uint32 SampleRate, float SampleSizeSeconds);
-		virtual ~FAudioPatchMixer() = default;
+		UE_API virtual ~FAudioPatchMixer() = default;
 
 		UE_API uint32 GetMaxBufferSize() const;
 		UE_API uint8  GetNumChannels() const;
@@ -48,7 +51,7 @@ namespace UE::PixelStreaming2
 	class FMixAudioTask : public FPixelStreamingTickableTask
 	{
 	public:
-		UE_API FMixAudioTask(FAudioCapturer* Capturer, TSharedPtr<FAudioPatchMixer> Mixer);
+		UE_API FMixAudioTask(TWeakPtr<FAudioCapturer> InCapturer, TSharedPtr<FAudioPatchMixer> InMixer);
 
 		virtual ~FMixAudioTask() = default;
 
@@ -61,14 +64,16 @@ namespace UE::PixelStreaming2
 		bool								  bIsRunning;
 		Audio::VectorOps::FAlignedFloatBuffer MixingBuffer;
 
-		FAudioCapturer*				 Capturer;
+		TWeakPtr<FAudioCapturer>	 Capturer;
 		TSharedPtr<FAudioPatchMixer> Mixer;
 	};
 
-	class FAudioCapturer
+	class FAudioCapturer : public TSharedFromThis<FAudioCapturer>
 	{
 	public:
-		static UE_API TSharedPtr<FAudioCapturer> Create(const int InSampleRate = 48000, const int InNumChannels = 2, const float InSampleSizeInSeconds = 0.5f);
+		template<typename T>
+		static TSharedPtr<T> Create(const int InSampleRate = 48000, const int InNumChannels = 2, const float InSampleSizeInSeconds = 0.5f);
+
 		virtual ~FAudioCapturer() = default;
 
 		// Mixed audio input will push its audio to an FPatchInputProxy for mixing
@@ -90,6 +95,10 @@ namespace UE::PixelStreaming2
 
 	protected:
 		UE_API FAudioCapturer(const int SampleRate = 48000, const int NumChannels = 2, const float SampleSizeInSeconds = 0.5f);
+		
+		// Required to be called by derived classes if the class is not constructed with FAudioCapturer::Create.
+		// It initializes members that require TSharedPtr to FAudioCapturer
+		UE_API void Initialize();
 
 		UE_API void OnDebugDumpAudioChanged(IConsoleVariable* Var);
 		UE_API void OnEnginePreExit();
@@ -114,6 +123,32 @@ namespace UE::PixelStreaming2
 
 		Audio::TSampleBuffer<int16_t> DebugDumpAudioBuffer;
 	};
+
+	template <typename T>
+	TSharedPtr<T> FAudioCapturer::Create(const int InSampleRate, const int InNumChannels, const float InSampleSizeInSeconds)
+	{
+		static_assert(std::is_base_of_v<FAudioCapturer, T>);
+		TSharedPtr<T> AudioCapturer(new T(InSampleRate, InNumChannels, InSampleSizeInSeconds));
+		AudioCapturer->Initialize();
+
+		FAudioDeviceManagerDelegates::OnAudioDeviceCreated.AddSP(AudioCapturer.ToSharedRef(), &T::CreateAudioProducer);
+		FAudioDeviceManagerDelegates::OnAudioDeviceDestroyed.AddSP(AudioCapturer.ToSharedRef(), &T::RemoveAudioProducer);
+
+		if (UPixelStreaming2PluginSettings::FDelegates* Delegates = UPixelStreaming2PluginSettings::Delegates())
+		{
+			Delegates->OnDebugDumpAudioChanged.AddSP(AudioCapturer.ToSharedRef(), &T::OnDebugDumpAudioChanged);
+
+			TWeakPtr<T> WeakAudioMixingCapturer = AudioCapturer;
+			FCoreDelegates::OnEnginePreExit.AddLambda([WeakAudioMixingCapturer]() {
+				if (TSharedPtr<T> AudioCapturer = WeakAudioMixingCapturer.Pin())
+				{
+					AudioCapturer->OnEnginePreExit();
+				}
+			});
+		}
+
+		return AudioCapturer;
+	}
 } // namespace UE::PixelStreaming2
 
 #undef UE_API
diff --git a/Engine/Plugins/Media/PixelStreaming2/Source/PixelStreaming2/Private/AudioCapturer.cpp b/Engine/Plugins/Media/PixelStreaming2/Source/PixelStreaming2/Private/AudioCapturer.cpp
--- a/Engine/Plugins/Media/PixelStreaming2/Source/PixelStreaming2/Private/AudioCapturer.cpp
+++ b/Engine/Plugins/Media/PixelStreaming2/Source/PixelStreaming2/Private/AudioCapturer.cpp
@@ -128,9 +128,9 @@ namespace UE::PixelStreaming2
 	/***************************************************
 	 *
 	 ****************************************************/
-	FMixAudioTask::FMixAudioTask(FAudioCapturer* Capturer, TSharedPtr<FAudioPatchMixer> Mixer)
-		: Capturer(Capturer)
-		, Mixer(Mixer)
+	FMixAudioTask::FMixAudioTask(TWeakPtr<FAudioCapturer> InCapturer, TSharedPtr<FAudioPatchMixer> InMixer)
+		: Capturer(InCapturer)
+		, Mixer(InMixer)
 	{
 		MixingBuffer.SetNumUninitialized(Mixer->GetMaxBufferSize());
 	}
@@ -160,7 +160,10 @@ namespace UE::PixelStreaming2
 			return;
 		}
 
-		Capturer->OnAudio(MixingBuffer.GetData(), NSamplesPopped, Mixer->GetNumChannels(), Mixer->GetSampleRate());
+		if (TSharedPtr<FAudioCapturer> PinnedCapturer = Capturer.Pin())
+		{
+			PinnedCapturer->OnAudio(MixingBuffer.GetData(), NSamplesPopped, Mixer->GetNumChannels(), Mixer->GetSampleRate());
+		}
 	}
 	const FString& FMixAudioTask::GetName() const
 	{
@@ -171,27 +174,26 @@ namespace UE::PixelStreaming2
 	/***************************************************
 	 *
 	 ****************************************************/
-	TSharedPtr<FAudioCapturer> FAudioCapturer::Create(const int InSampleRate, const int InNumChannels, const float InSampleSizeInSeconds)
-	{
-		TSharedPtr<FAudioCapturer> AudioCapturer(new FAudioCapturer(InSampleRate, InNumChannels, InSampleSizeInSeconds));
 
-		FAudioDeviceManagerDelegates::OnAudioDeviceCreated.AddSP(AudioCapturer.ToSharedRef(), &FAudioCapturer::CreateAudioProducer);
-		FAudioDeviceManagerDelegates::OnAudioDeviceDestroyed.AddSP(AudioCapturer.ToSharedRef(), &FAudioCapturer::RemoveAudioProducer);
-
-		if (UPixelStreaming2PluginSettings::FDelegates* Delegates = UPixelStreaming2PluginSettings::Delegates())
+	void FAudioCapturer::Initialize()
+	{
+		// If No engine. Possibly running editor tests
+		if (GEngine)
 		{
-			Delegates->OnDebugDumpAudioChanged.AddSP(AudioCapturer.ToSharedRef(), &FAudioCapturer::OnDebugDumpAudioChanged);
-
-			TWeakPtr<FAudioCapturer> WeakAudioMixingCapturer = AudioCapturer;
-			FCoreDelegates::OnEnginePreExit.AddLambda([WeakAudioMixingCapturer]() {
-				if (TSharedPtr<FAudioCapturer> AudioCapturer = WeakAudioMixingCapturer.Pin())
-				{
-					AudioCapturer->OnEnginePreExit();
-				}
-			});
+			// subscribe to audio data
+			Mixer = MakeShared<FAudioPatchMixer>(NumChannels, SampleRate, SampleSizeSeconds);
+			MixerTask = FPixelStreamingTickableTask::Create<FMixAudioTask>(AsWeak(), Mixer);
+			if (FAudioDeviceManager* AudioDeviceManager = FAudioDeviceManager::Get(); AudioDeviceManager != nullptr)
+			{
+				TWeakPtr<FAudioCapturer> LocalAudioCapturer = AsWeak();
+				AudioDeviceManager->IterateOverAllDevices([LocalAudioCapturer](Audio::FDeviceId AudioDeviceId, FAudioDevice*) {
+					if (TSharedPtr<FAudioCapturer> Pin = LocalAudioCapturer.Pin())
+					{
+						Pin->CreateAudioProducer(AudioDeviceId);
+					}
+				});
+			}
 		}
-
-		return AudioCapturer;
 	}
 
 	FAudioCapturer::FAudioCapturer(const int SampleRate, const int NumChannels, const float NumSampleSizeInSeconds)
@@ -199,20 +201,6 @@ namespace UE::PixelStreaming2
 		, NumChannels(NumChannels)
 		, SampleSizeSeconds(NumSampleSizeInSeconds)
 	{
-		// subscribe to audio data
-		if (!GEngine)
-		{
-			// No engine. Possibly running editor tests
-			return;
-		}
-		Mixer = MakeShared<FAudioPatchMixer>(NumChannels, SampleRate, SampleSizeSeconds);
-		MixerTask = FPixelStreamingTickableTask::Create<FMixAudioTask>(this, Mixer);
-		if (FAudioDeviceManager* AudioDeviceManager = FAudioDeviceManager::Get(); AudioDeviceManager != nullptr)
-		{
-			AudioDeviceManager->IterateOverAllDevices([this](Audio::FDeviceId AudioDeviceId, FAudioDevice*) {
-				CreateAudioProducer(AudioDeviceId);
-			});
-		}
 	}
 
 	void FAudioCapturer::CreateAudioProducer(Audio::FDeviceId AudioDeviceId)
@@ -240,11 +228,10 @@ namespace UE::PixelStreaming2
 
 	void FAudioCapturer::OnAudio(const float* AudioData, int32 InNumSamples, int32 InNumChannels, int32 InSampleRate)
 	{
-		// Note: TSampleBuffer takes in AudioData as float* and internally converts to int16
-		Audio::TSampleBuffer<int16_t> Buffer(AudioData, InNumSamples, InNumChannels, SampleRate);
-
 		if (UPixelStreaming2PluginSettings::CVarDebugDumpAudio.GetValueOnAnyThread())
 		{
+			// Note: TSampleBuffer takes in AudioData as float* and internally converts to int16
+			Audio::TSampleBuffer<int16_t> Buffer(AudioData, InNumSamples, InNumChannels, SampleRate);
 			DebugDumpAudioBuffer.Append(Buffer.GetData(), Buffer.GetNumSamples(), Buffer.GetNumChannels(), Buffer.GetSampleRate());
 		}
 
diff --git a/Engine/Plugins/Media/PixelStreaming2/Source/PixelStreaming2/Private/VideoProducerViewportBase.cpp b/Engine/Plugins/Media/PixelStreaming2/Source/PixelStreaming2/Private/VideoProducerViewportBase.cpp
--- a/Engine/Plugins/Media/PixelStreaming2/Source/PixelStreaming2/Private/VideoProducerViewportBase.cpp
+++ b/Engine/Plugins/Media/PixelStreaming2/Source/PixelStreaming2/Private/VideoProducerViewportBase.cpp
@@ -91,7 +91,7 @@ namespace UE::PixelStreaming2
 		TRefCountPtr<IPooledRenderTarget> OutTexture;
 		{
 			FRDGBuilder GraphBuilder(FRHICommandListImmediate::Get());
-			FRDGTextureRef InputTexture = GraphBuilder.RegisterExternalTexture(CreateRenderTarget(BackBuffer, *Window.GetTitle().ToString()));
+			FRDGTextureRef InputTexture = GraphBuilder.RegisterExternalTexture(CreateRenderTarget(BackBuffer, TEXT("VideoProducerViewportBaseInput")));
 			// TODO (william.belcher): The cropped texture having an 8bit pixel format matches the full editor video producer code path but we may need to reconsider this if we want to add 10bit video support
 			FRDGTextureRef CroppedTexture = GraphBuilder.CreateTexture(FRDGTextureDesc::Create2D(VirtualSize, EPixelFormat::PF_B8G8R8A8, FClearValueBinding::None, GetTexCreateFlags()), TEXT("VideoProducerViewportBaseCropped"));
 
diff --git a/Engine/Plugins/Media/PixelStreaming2/Source/PixelStreaming2Editor/Private/VideoProducerBackBufferComposited.cpp b/Engine/Plugins/Media/PixelStreaming2/Source/PixelStreaming2Editor/Private/VideoProducerBackBufferComposited.cpp
--- a/Engine/Plugins/Media/PixelStreaming2/Source/PixelStreaming2Editor/Private/VideoProducerBackBufferComposited.cpp
+++ b/Engine/Plugins/Media/PixelStreaming2/Source/PixelStreaming2Editor/Private/VideoProducerBackBufferComposited.cpp
@@ -102,8 +102,6 @@ namespace UE::EditorPixelStreaming2
 		 * Finally, the composited frame is extracted from the RDG pipeline and we send it on its way through the PixelCapturer
 		 *
 		 */
-		UE_LOG(LogPixelStreaming2BackBufferComposited, Verbose, TEXT("Type: %s"), *SlateWindow.GetTitle().ToString());
-
 		TSharedPtr<FVideoProducerUserData> UserData = MakeShared<FVideoProducerUserData>();
 		UserData->ProducerName = TEXT("FVideoProducerBackBufferComposited");
 		UserData->ProductionBeginCycles = FPlatformTime::Cycles64();
@@ -128,7 +126,7 @@ namespace UE::EditorPixelStreaming2
 				FRDGBuilder GraphBuilder(FRHICommandListImmediate::Get());
 
 				// Register an external RDG texture from the provided frame buffer
-				FRDGTextureRef InputTexture = GraphBuilder.RegisterExternalTexture(CreateRenderTarget(FrameBuffer, *SlateWindow.GetTitle().ToString()));
+				FRDGTextureRef InputTexture = GraphBuilder.RegisterExternalTexture(CreateRenderTarget(FrameBuffer, TEXT("VideoProducerBackBufferCompositedInput")));
 
 				// Create an internal RDG texture with the same extent and format as the source
 				FRDGTextureRef OutputTexture = GraphBuilder.CreateTexture(FRDGTextureDesc::Create2D(InputTexture->Desc.Extent, InputTexture->Desc.Format, FClearValueBinding::None, GetTexCreateFlags()), TEXT("VideoProducerBackBufferCompositedStaging"));
diff --git a/Engine/Plugins/Media/PixelStreaming2/Source/PixelStreaming2Input/Private/InputHandler.cpp b/Engine/Plugins/Media/PixelStreaming2/Source/PixelStreaming2Input/Private/InputHandler.cpp
--- a/Engine/Plugins/Media/PixelStreaming2/Source/PixelStreaming2Input/Private/InputHandler.cpp
+++ b/Engine/Plugins/Media/PixelStreaming2/Source/PixelStreaming2Input/Private/InputHandler.cpp
@@ -52,6 +52,8 @@ namespace UE::PixelStreaming2Input
 
 	FPixelStreaming2InputHandler::~FPixelStreaming2InputHandler()
 	{
+		IModularFeatures::Get().UnregisterModularFeature(GetModularFeatureName(), this);
+
 		if (UPixelStreaming2PluginSettings::FDelegates* Delegates = UPixelStreaming2PluginSettings::Delegates())
 		{
 			Delegates->OnInputKeyFilterChanged.Remove(OnInputKeyFilterChangedHandle);
diff --git a/Engine/Plugins/Media/PixelStreaming2/Source/PixelStreaming2RTC/Internal/EpicRtcAudioCapturer.h b/Engine/Plugins/Media/PixelStreaming2/Source/PixelStreaming2RTC/Internal/EpicRtcAudioCapturer.h
--- a/Engine/Plugins/Media/PixelStreaming2/Source/PixelStreaming2RTC/Internal/EpicRtcAudioCapturer.h
+++ b/Engine/Plugins/Media/PixelStreaming2/Source/PixelStreaming2RTC/Internal/EpicRtcAudioCapturer.h
@@ -15,17 +15,20 @@ namespace UE::PixelStreaming2
 	class FEpicRtcAudioCapturer : public FAudioCapturer
 	{
 	public:
-		static UE_API TSharedPtr<FEpicRtcAudioCapturer> Create();
-		virtual ~FEpicRtcAudioCapturer() = default;
+		UE_API virtual ~FEpicRtcAudioCapturer() = default;
 
 		// Override the push audio method as EpicRtc needs the broadcasted audio to be in 10ms chunks
 		UE_API virtual void PushAudio(const float* AudioData, int32 NumSamples, int32 NumChannels, int32 SampleRate) override;
 
 	protected:
-		FEpicRtcAudioCapturer() = default;
+		UE_API FEpicRtcAudioCapturer(const int SampleRate = 48000, const int NumChannels = 2, const float SampleSizeInSeconds = 0.5f)
+			: FAudioCapturer(SampleRate, NumChannels, SampleSizeInSeconds) {}
 
 	private:
 		TArray<int16_t> RecordingBuffer;
+
+		// Needed so FAudioCapturer::Create can access the constructor
+		friend FAudioCapturer;
 	};
 } // namespace UE::PixelStreaming2
 
diff --git a/Engine/Plugins/Media/PixelStreaming2/Source/PixelStreaming2RTC/Private/EpicRtcAudioCapturer.cpp b/Engine/Plugins/Media/PixelStreaming2/Source/PixelStreaming2RTC/Private/EpicRtcAudioCapturer.cpp
--- a/Engine/Plugins/Media/PixelStreaming2/Source/PixelStreaming2RTC/Private/EpicRtcAudioCapturer.cpp
+++ b/Engine/Plugins/Media/PixelStreaming2/Source/PixelStreaming2RTC/Private/EpicRtcAudioCapturer.cpp
@@ -8,31 +8,13 @@
 
 namespace UE::PixelStreaming2
 {
-	TSharedPtr<FEpicRtcAudioCapturer> FEpicRtcAudioCapturer::Create()
-	{
-		TSharedPtr<FEpicRtcAudioCapturer> AudioMixingCapturer(new FEpicRtcAudioCapturer());
-
-		FAudioDeviceManagerDelegates::OnAudioDeviceCreated.AddSP(AudioMixingCapturer.ToSharedRef(), &FEpicRtcAudioCapturer::CreateAudioProducer);
-		FAudioDeviceManagerDelegates::OnAudioDeviceDestroyed.AddSP(AudioMixingCapturer.ToSharedRef(), &FEpicRtcAudioCapturer::RemoveAudioProducer);
-
-		if (UPixelStreaming2PluginSettings::FDelegates* Delegates = UPixelStreaming2PluginSettings::Delegates())
-		{
-			Delegates->OnDebugDumpAudioChanged.AddSP(AudioMixingCapturer.ToSharedRef(), &FEpicRtcAudioCapturer::OnDebugDumpAudioChanged);
-
-			TWeakPtr<FEpicRtcAudioCapturer> WeakAudioMixingCapturer = AudioMixingCapturer;
-			FCoreDelegates::OnEnginePreExit.AddLambda([WeakAudioMixingCapturer]() {
-				if (TSharedPtr<FEpicRtcAudioCapturer> AudioMixingCapturer = WeakAudioMixingCapturer.Pin())
-				{
-					AudioMixingCapturer->OnEnginePreExit();
-				}
-			});
-		}
-
-		return AudioMixingCapturer;
-	}
-
 	void FEpicRtcAudioCapturer::PushAudio(const float* AudioData, int32 InNumSamples, int32 InNumChannels, const int32 InSampleRate)
 	{
+		check(AudioData != nullptr);
+		check(InNumSamples > 0);
+		check(InNumChannels > 0);
+		check(InSampleRate > 0);
+		check(InSampleRate == SampleRate);
 		Audio::TSampleBuffer<int16_t> Buffer(AudioData, InNumSamples, InNumChannels, SampleRate);
 		RecordingBuffer.Append(Buffer.GetData(), Buffer.GetNumSamples());
 
@@ -40,7 +22,7 @@ namespace UE::PixelStreaming2
 		const size_t BytesPerFrame = NumChannels * sizeof(int16_t);
 
 		// Feed in 10ms chunks
-		while (RecordingBuffer.Num() > SamplesPer10Ms)
+		while (RecordingBuffer.Num() >= SamplesPer10Ms)
 		{
 			OnAudioBuffer.Broadcast(RecordingBuffer.GetData(), SamplesPer10Ms, NumChannels, SampleRate);
 
diff --git a/Engine/Plugins/Media/PixelStreaming2/Source/PixelStreaming2RTC/Private/EpicRtcStreamer.cpp b/Engine/Plugins/Media/PixelStreaming2/Source/PixelStreaming2RTC/Private/EpicRtcStreamer.cpp
--- a/Engine/Plugins/Media/PixelStreaming2/Source/PixelStreaming2RTC/Private/EpicRtcStreamer.cpp
+++ b/Engine/Plugins/Media/PixelStreaming2/Source/PixelStreaming2RTC/Private/EpicRtcStreamer.cpp
@@ -2,6 +2,7 @@
 
 #include "EpicRtcStreamer.h"
 
+#include "AudioCapturer.h"
 #include "PixelStreaming2PluginSettings.h"
 #include "PixelStreaming2Delegates.h"
 #include "Engine/GameEngine.h"
@@ -257,8 +258,7 @@ namespace UE::PixelStreaming2
 			UE_LOGFMT(LogPixelStreaming2RTC, Log, "Streamer is already streaming. Ignoring subsequent call to StartStreaming!");
 			return;
 		}
-		StreamState = EStreamState::Connecting;
-
+		
 		if (CurrentSignallingServerURL.IsEmpty())
 		{
 			UE_LOGFMT(LogPixelStreaming2RTC, Warning, "Attempted to start streamer ({0}) but no signalling server URL has been set. Use Streamer->SetConnectionURL(URL) or -PixelStreamingConnectionURL=", StreamerId);
@@ -267,12 +267,14 @@ namespace UE::PixelStreaming2
 
 		check(EpicRtcConference.IsValid());
 
+		StreamState = EStreamState::Connecting;
+
 		ReconnectTimer->Stop();
 
 		VideoCapturer->ResetFrameCapturer();
 
 		TickableTasks = FPixelStreaming2RTCModule::GetModule()->GetSharedTickableTasks();
-		AudioCapturer = FEpicRtcAudioCapturer::Create();
+		AudioCapturer = FAudioCapturer::Create<FEpicRtcAudioCapturer>();
 
 		{
 			FScopeLock Lock(&CustomAudioProducersCS);
diff --git a/Engine/Plugins/Media/PixelStreaming2/Source/PixelStreaming2Settings/Private/PixelStreaming2PluginSettings.cpp b/Engine/Plugins/Media/PixelStreaming2/Source/PixelStreaming2Settings/Private/PixelStreaming2PluginSettings.cpp
--- a/Engine/Plugins/Media/PixelStreaming2/Source/PixelStreaming2Settings/Private/PixelStreaming2PluginSettings.cpp
+++ b/Engine/Plugins/Media/PixelStreaming2/Source/PixelStreaming2Settings/Private/PixelStreaming2PluginSettings.cpp
@@ -890,6 +890,11 @@ void UPixelStreaming2PluginSettings::PostEditChangeProperty(FPropertyChangedEven
 {
 	Super::PostEditChangeProperty(PropertyChangedEvent);
 
+	if (!PropertyChangedEvent.Property)
+	{
+		return;
+	}
+
 	FString PropertyName = PropertyChangedEvent.Property->GetNameCPP();
 
 	FString CVarName;
