[NVCodecs, PixelStreaming2] Fixes:
- Fix: Crash when NVDEC is unable to load a function pointer. The fatal message has been downgraded to a warning with an additional check for if all functions were successully loaded
- Fix: UIInteractions received from the browser were not properly respecting the descriptors length
- Fix: Crash when OnAudioTrackUpdate was executed after StopStreaming

diff --git a/Engine/Plugins/Experimental/AVCodecs/NVCodecs/Source/NVDEC/Private/NVDEC.cpp b/Engine/Plugins/Experimental/AVCodecs/NVCodecs/Source/NVDEC/Private/NVDEC.cpp
--- a/Engine/Plugins/Experimental/AVCodecs/NVCodecs/Source/NVDEC/Private/NVDEC.cpp
+++ b/Engine/Plugins/Experimental/AVCodecs/NVCodecs/Source/NVDEC/Private/NVDEC.cpp
@@ -13,14 +13,16 @@ REGISTER_TYPEID(FNVDEC);
 	Name = (Name##Ptr)FPlatformProcess::GetDllExport(DllHandle, TEXT(#Name)); \
 	if (Name == nullptr) \
 	{ \
-		FAVResult::Log(EAVResult::Fatal, TEXT("Failed to get dll export function: " #Name), TEXT("NVDEC")); \
+		bAllFunctionsLoaded = false; \
+		FAVResult::Log(EAVResult::Warning, TEXT("Failed to get dll export function: " #Name), TEXT("NVDEC")); \
 	}
 
 #define LOAD_FROM_DLL64(Name) \
 	Name = (Name##Ptr)FPlatformProcess::GetDllExport(DllHandle, TEXT(#Name) "64"); \
 	if (Name == nullptr) \
 	{ \
-		FAVResult::Log(EAVResult::Fatal, TEXT("Failed to get dll export function: " #Name), TEXT("NVDEC")); \
+		bAllFunctionsLoaded = false; \
+		FAVResult::Log(EAVResult::Warning, TEXT("Failed to get dll export function: " #Name), TEXT("NVDEC")); \
 	}
 
 FNVDEC::FNVDEC()
@@ -50,6 +52,7 @@ FNVDEC::FNVDEC()
 		DllHandle = FPlatformProcess::GetDllHandle(DllName);
 		if (DllHandle != nullptr)
 		{
+			bool bAllFunctionsLoaded = true;
 			LOAD_FROM_DLL(cuvidCreateDecoder);
 			LOAD_FROM_DLL(cuvidParseVideoData);
 			LOAD_FROM_DLL(cuvidCreateVideoParser);
@@ -61,14 +64,16 @@ FNVDEC::FNVDEC()
 			LOAD_FROM_DLL(cuvidGetDecodeStatus);
 			LOAD_FROM_DLL(cuvidReconfigureDecoder);
 			LOAD_FROM_DLL(cuvidDestroyDecoder);
-
-#if PLATFORM_64BITS
 			LOAD_FROM_DLL64(cuvidMapVideoFrame);
 			LOAD_FROM_DLL64(cuvidUnmapVideoFrame);
-#else
-			LOAD_FROM_DLL(cuvidMapVideoFrame);
-			LOAD_FROM_DLL(cuvidUnmapVideoFrame);
-#endif
+
+			if (bAllFunctionsLoaded)
+			{
+				return;
+			}
+
+			FPlatformProcess::FreeDllHandle(DllHandle);
+			DllHandle = nullptr;
 		}
 		else
 		{
diff --git a/Engine/Plugins/Media/PixelStreaming2/Source/PixelStreaming2RTC/Private/EpicRtcStreamer.cpp b/Engine/Plugins/Media/PixelStreaming2/Source/PixelStreaming2RTC/Private/EpicRtcStreamer.cpp
--- a/Engine/Plugins/Media/PixelStreaming2/Source/PixelStreaming2RTC/Private/EpicRtcStreamer.cpp
+++ b/Engine/Plugins/Media/PixelStreaming2/Source/PixelStreaming2RTC/Private/EpicRtcStreamer.cpp
@@ -375,11 +375,27 @@ namespace UE::PixelStreaming2
 
 	void FEpicRtcStreamer::OnUIInteraction(FMemoryReader Ar)
 	{
-		FString Res;
-		Res.GetCharArray().SetNumUninitialized(Ar.TotalSize() / 2 + 1);
-		Ar.Serialize(Res.GetCharArray().GetData(), Ar.TotalSize());
+		if (Ar.TotalSize() < 2)
+		{
+			UE_LOGFMT(LogPixelStreaming2RTC, Warning, "FEpicRtcStreamer[\"{0}\"]::OnUIInteraction: Malformed message. Expected to contain at least 2 bytes representing the descriptor length in characters", StreamerId);
+			return;
+		}
 
-		FString Descriptor = Res.Mid(1);
+		uint16 StringLength = 0;
+		Ar << StringLength;
+
+		if (Ar.TotalSize() < static_cast<int64>(2 + StringLength * sizeof(TCHAR)))
+		{
+			UE_LOGFMT(LogPixelStreaming2RTC, Warning, "FEpicRtcStreamer[\"{0}\"]::OnUIInteraction: Malformed message. Descriptor length of {1} characters exceeds remaining message size", StreamerId, StringLength);
+			return;
+		}
+
+		FString Descriptor;
+		Descriptor.GetCharArray().SetNumUninitialized(StringLength + 1 /* +1 for null terminator */);
+		// Read archive into the string's char array
+		Ar.Serialize(Descriptor.GetCharArray().GetData(), StringLength * sizeof(TCHAR));
+		// Null terminate
+		Descriptor.GetCharArray()[StringLength] = TEXT('\0');
 
 		UE_LOGFMT(LogPixelStreaming2RTC, Verbose, "FEpicRtcStreamer[\"{0}\"]::OnUIInteraction({1})", StreamerId, Descriptor);
 		InputComponents.Apply([Descriptor](uintptr_t Key, UPixelStreaming2Input* Value) { 
@@ -1652,6 +1668,13 @@ namespace UE::PixelStreaming2
 		AudioTrackPlayerIdMap.Add(reinterpret_cast<uintptr_t>(AudioTrack), ParticipantId);
 		if (ParticipantId == StreamerId)
 		{
+			// A user may hit a race condition if they call StopStreaming (which resets AudioCapturer) before this callback 
+			// has a chance to  be executed by the task thread. We opt to create an AudioCapturer here instead of early
+			// returning to ensure AudioSource is always valid for the streamer.
+			if (!AudioCapturer)
+			{
+				AudioCapturer = FAudioCapturer::Create<FEpicRtcAudioCapturer>();
+			}
 			// NOTE: we don't broadcast the audio track open for a streamer
 			AudioSource = FEpicRtcAudioSource::Create(AudioTrack, AudioCapturer);
 		}
