[PS2] Fix: PIEStreamer crashing due to unsafe pointer access

diff --git a/Engine/Plugins/Media/PixelStreaming2/Source/PixelStreaming2Editor/Private/PixelStreaming2EditorModule.cpp b/Engine/Plugins/Media/PixelStreaming2/Source/PixelStreaming2Editor/Private/PixelStreaming2EditorModule.cpp
--- a/Engine/Plugins/Media/PixelStreaming2/Source/PixelStreaming2Editor/Private/PixelStreaming2EditorModule.cpp
+++ b/Engine/Plugins/Media/PixelStreaming2/Source/PixelStreaming2Editor/Private/PixelStreaming2EditorModule.cpp
@@ -564,21 +564,13 @@ namespace UE::EditorPixelStreaming2
 			return;
 		}
 
-		FWorldContext* PIEWorldContext = GEditor->GetPIEWorldContext();
-		if (!PIEWorldContext)
+		UGameViewportClient* Viewport = GEditor->GetNextPIEViewport(nullptr);
+		if (!Viewport)
 		{
-			UE_LOGFMT(LogPixelStreaming2Editor, Warning, "Failed to get editor PIE world context. PIE streamer will not stream!");
+			UE_LOGFMT(LogPixelStreaming2Editor, Warning, "Failed to get PIE viewport. PIE streamer will not stream!");
 			return;
 		}
-
-		UWorld* World = PIEWorldContext->World();
-		if (!World)
-		{
-			UE_LOGFMT(LogPixelStreaming2Editor, Warning, "Failed to get editor PIE world. PIE streamer will not stream!");
-			return;
-		}
-
-		UGameViewportClient* Viewport = World->GetGameViewport();
+		
 		InputHandler->SetTargetViewport(Viewport->GetGameViewportWidget());
 		InputHandler->SetTargetWindow(Viewport->GetWindow());
 		InputHandler->SetInputType(EPixelStreaming2InputType::RouteToWindow);
