PixelStreaming2: Added the allowed platforms for the referenced AVCodecsCore plugin

#rnx

diff --git a/Engine/Plugins/Media/PixelStreaming2/PixelStreaming2.uplugin b/Engine/Plugins/Media/PixelStreaming2/PixelStreaming2.uplugin
--- a/Engine/Plugins/Media/PixelStreaming2/PixelStreaming2.uplugin
+++ b/Engine/Plugins/Media/PixelStreaming2/PixelStreaming2.uplugin
@@ -90,10 +90,11 @@
 			"Enabled": true,
 			"PlatformAllowList": [ "Win64" ]
 		},
-		{
-			"Name": "AVCodecsCore",
-			"Enabled": true
-		},
+    {
+      "Name": "AVCodecsCore",
+      "Enabled": true,
+      "PlatformAllowList": [ "Win64", "Linux", "Mac" ]
+    },
 		{
 			"Name": "NVCodecs",
 			"Enabled": true,
