[PS2] Fix: crash when calling PushAudio.

Convert FAudioCapture to be a TSharedPtr rather than raw pointer.

diff --git a/Engine/Plugins/Media/PixelStreaming2/Source/PixelStreaming2/Internal/AudioCapturer.h b/Engine/Plugins/Media/PixelStreaming2/Source/PixelStreaming2/Internal/AudioCapturer.h
--- a/Engine/Plugins/Media/PixelStreaming2/Source/PixelStreaming2/Internal/AudioCapturer.h
+++ b/Engine/Plugins/Media/PixelStreaming2/Source/PixelStreaming2/Internal/AudioCapturer.h
@@ -3,7 +3,10 @@
 #pragma once
 
 #include "DSP/MultithreadedPatching.h"
+#include "AudioDeviceManager.h"
 #include "AudioProducer.h"
+#include "Misc/CoreDelegates.h"
+#include "PixelStreaming2PluginSettings.h"
 #include "SampleBuffer.h"
 #include "TickableTask.h"
 
@@ -48,7 +51,7 @@ namespace UE::PixelStreaming2
 	class FMixAudioTask : public FPixelStreamingTickableTask
 	{
 	public:
-		UE_API FMixAudioTask(FAudioCapturer* Capturer, TSharedPtr<FAudioPatchMixer> Mixer);
+		UE_API FMixAudioTask(TWeakPtr<FAudioCapturer> InCapturer, TSharedPtr<FAudioPatchMixer> InMixer);
 
 		virtual ~FMixAudioTask() = default;
 
@@ -61,14 +64,16 @@ namespace UE::PixelStreaming2
 		bool								  bIsRunning;
 		Audio::VectorOps::FAlignedFloatBuffer MixingBuffer;
 
-		FAudioCapturer*				 Capturer;
+		TWeakPtr<FAudioCapturer>	 Capturer;
 		TSharedPtr<FAudioPatchMixer> Mixer;
 	};
 
-	class FAudioCapturer
+	class FAudioCapturer : public TSharedFromThis<FAudioCapturer>
 	{
 	public:
-		static UE_API TSharedPtr<FAudioCapturer> Create(const int InSampleRate = 48000, const int InNumChannels = 2, const float InSampleSizeInSeconds = 0.5f);
+		template<typename T>
+		static TSharedPtr<T> Create(const int InSampleRate = 48000, const int InNumChannels = 2, const float InSampleSizeInSeconds = 0.5f);
+
 		virtual ~FAudioCapturer() = default;
 
 		// Mixed audio input will push its audio to an FPatchInputProxy for mixing
@@ -90,6 +95,10 @@ namespace UE::PixelStreaming2
 
 	protected:
 		UE_API FAudioCapturer(const int SampleRate = 48000, const int NumChannels = 2, const float SampleSizeInSeconds = 0.5f);
+		
+		// Required to be called by derived classes if the class is not constructed with FAudioCapturer::Create.
+		// It initializes members that require TSharedPtr to FAudioCapturer
+		UE_API void Initialize();
 
 		UE_API void OnDebugDumpAudioChanged(IConsoleVariable* Var);
 		UE_API void OnEnginePreExit();
@@ -114,6 +123,32 @@ namespace UE::PixelStreaming2
 
 		Audio::TSampleBuffer<int16_t> DebugDumpAudioBuffer;
 	};
+
+	template <typename T>
+	TSharedPtr<T> FAudioCapturer::Create(const int InSampleRate, const int InNumChannels, const float InSampleSizeInSeconds)
+	{
+		static_assert(std::is_base_of_v<FAudioCapturer, T>);
+		TSharedPtr<T> AudioCapturer(new T(InSampleRate, InNumChannels, InSampleSizeInSeconds));
+		AudioCapturer->Initialize();
+
+		FAudioDeviceManagerDelegates::OnAudioDeviceCreated.AddSP(AudioCapturer.ToSharedRef(), &T::CreateAudioProducer);
+		FAudioDeviceManagerDelegates::OnAudioDeviceDestroyed.AddSP(AudioCapturer.ToSharedRef(), &T::RemoveAudioProducer);
+
+		if (UPixelStreaming2PluginSettings::FDelegates* Delegates = UPixelStreaming2PluginSettings::Delegates())
+		{
+			Delegates->OnDebugDumpAudioChanged.AddSP(AudioCapturer.ToSharedRef(), &T::OnDebugDumpAudioChanged);
+
+			TWeakPtr<T> WeakAudioMixingCapturer = AudioCapturer;
+			FCoreDelegates::OnEnginePreExit.AddLambda([WeakAudioMixingCapturer]() {
+				if (TSharedPtr<T> AudioCapturer = WeakAudioMixingCapturer.Pin())
+				{
+					AudioCapturer->OnEnginePreExit();
+				}
+			});
+		}
+
+		return AudioCapturer;
+	}
 } // namespace UE::PixelStreaming2
 
 #undef UE_API
diff --git a/Engine/Plugins/Media/PixelStreaming2/Source/PixelStreaming2/Private/AudioCapturer.cpp b/Engine/Plugins/Media/PixelStreaming2/Source/PixelStreaming2/Private/AudioCapturer.cpp
--- a/Engine/Plugins/Media/PixelStreaming2/Source/PixelStreaming2/Private/AudioCapturer.cpp
+++ b/Engine/Plugins/Media/PixelStreaming2/Source/PixelStreaming2/Private/AudioCapturer.cpp
@@ -128,9 +128,9 @@ namespace UE::PixelStreaming2
 	/***************************************************
 	 *
 	 ****************************************************/
-	FMixAudioTask::FMixAudioTask(FAudioCapturer* Capturer, TSharedPtr<FAudioPatchMixer> Mixer)
-		: Capturer(Capturer)
-		, Mixer(Mixer)
+	FMixAudioTask::FMixAudioTask(TWeakPtr<FAudioCapturer> InCapturer, TSharedPtr<FAudioPatchMixer> InMixer)
+		: Capturer(InCapturer)
+		, Mixer(InMixer)
 	{
 		MixingBuffer.SetNumUninitialized(Mixer->GetMaxBufferSize());
 	}
@@ -160,7 +160,10 @@ namespace UE::PixelStreaming2
 			return;
 		}
 
-		Capturer->OnAudio(MixingBuffer.GetData(), NSamplesPopped, Mixer->GetNumChannels(), Mixer->GetSampleRate());
+		if (TSharedPtr<FAudioCapturer> PinnedCapturer = Capturer.Pin())
+		{
+			PinnedCapturer->OnAudio(MixingBuffer.GetData(), NSamplesPopped, Mixer->GetNumChannels(), Mixer->GetSampleRate());
+		}
 	}
 	const FString& FMixAudioTask::GetName() const
 	{
@@ -171,27 +174,26 @@ namespace UE::PixelStreaming2
 	/***************************************************
 	 *
 	 ****************************************************/
-	TSharedPtr<FAudioCapturer> FAudioCapturer::Create(const int InSampleRate, const int InNumChannels, const float InSampleSizeInSeconds)
-	{
-		TSharedPtr<FAudioCapturer> AudioCapturer(new FAudioCapturer(InSampleRate, InNumChannels, InSampleSizeInSeconds));
 
-		FAudioDeviceManagerDelegates::OnAudioDeviceCreated.AddSP(AudioCapturer.ToSharedRef(), &FAudioCapturer::CreateAudioProducer);
-		FAudioDeviceManagerDelegates::OnAudioDeviceDestroyed.AddSP(AudioCapturer.ToSharedRef(), &FAudioCapturer::RemoveAudioProducer);
-
-		if (UPixelStreaming2PluginSettings::FDelegates* Delegates = UPixelStreaming2PluginSettings::Delegates())
+	void FAudioCapturer::Initialize()
+	{
+		// If No engine. Possibly running editor tests
+		if (GEngine)
 		{
-			Delegates->OnDebugDumpAudioChanged.AddSP(AudioCapturer.ToSharedRef(), &FAudioCapturer::OnDebugDumpAudioChanged);
-
-			TWeakPtr<FAudioCapturer> WeakAudioMixingCapturer = AudioCapturer;
-			FCoreDelegates::OnEnginePreExit.AddLambda([WeakAudioMixingCapturer]() {
-				if (TSharedPtr<FAudioCapturer> AudioCapturer = WeakAudioMixingCapturer.Pin())
-				{
-					AudioCapturer->OnEnginePreExit();
-				}
-			});
+			// subscribe to audio data
+			Mixer = MakeShared<FAudioPatchMixer>(NumChannels, SampleRate, SampleSizeSeconds);
+			MixerTask = FPixelStreamingTickableTask::Create<FMixAudioTask>(AsWeak(), Mixer);
+			if (FAudioDeviceManager* AudioDeviceManager = FAudioDeviceManager::Get(); AudioDeviceManager != nullptr)
+			{
+				TWeakPtr<FAudioCapturer> LocalAudioCapturer = AsWeak();
+				AudioDeviceManager->IterateOverAllDevices([LocalAudioCapturer](Audio::FDeviceId AudioDeviceId, FAudioDevice*) {
+					if (TSharedPtr<FAudioCapturer> Pin = LocalAudioCapturer.Pin())
+					{
+						Pin->CreateAudioProducer(AudioDeviceId);
+					}
+				});
+			}
 		}
-
-		return AudioCapturer;
 	}
 
 	FAudioCapturer::FAudioCapturer(const int SampleRate, const int NumChannels, const float NumSampleSizeInSeconds)
@@ -199,20 +201,6 @@ namespace UE::PixelStreaming2
 		, NumChannels(NumChannels)
 		, SampleSizeSeconds(NumSampleSizeInSeconds)
 	{
-		// subscribe to audio data
-		if (!GEngine)
-		{
-			// No engine. Possibly running editor tests
-			return;
-		}
-		Mixer = MakeShared<FAudioPatchMixer>(NumChannels, SampleRate, SampleSizeSeconds);
-		MixerTask = FPixelStreamingTickableTask::Create<FMixAudioTask>(this, Mixer);
-		if (FAudioDeviceManager* AudioDeviceManager = FAudioDeviceManager::Get(); AudioDeviceManager != nullptr)
-		{
-			AudioDeviceManager->IterateOverAllDevices([this](Audio::FDeviceId AudioDeviceId, FAudioDevice*) {
-				CreateAudioProducer(AudioDeviceId);
-			});
-		}
 	}
 
 	void FAudioCapturer::CreateAudioProducer(Audio::FDeviceId AudioDeviceId)
@@ -240,11 +228,10 @@ namespace UE::PixelStreaming2
 
 	void FAudioCapturer::OnAudio(const float* AudioData, int32 InNumSamples, int32 InNumChannels, int32 InSampleRate)
 	{
-		// Note: TSampleBuffer takes in AudioData as float* and internally converts to int16
-		Audio::TSampleBuffer<int16_t> Buffer(AudioData, InNumSamples, InNumChannels, SampleRate);
-
 		if (UPixelStreaming2PluginSettings::CVarDebugDumpAudio.GetValueOnAnyThread())
 		{
+			// Note: TSampleBuffer takes in AudioData as float* and internally converts to int16
+			Audio::TSampleBuffer<int16_t> Buffer(AudioData, InNumSamples, InNumChannels, SampleRate);
 			DebugDumpAudioBuffer.Append(Buffer.GetData(), Buffer.GetNumSamples(), Buffer.GetNumChannels(), Buffer.GetSampleRate());
 		}
 
diff --git a/Engine/Plugins/Media/PixelStreaming2/Source/PixelStreaming2RTC/Internal/EpicRtcAudioCapturer.h b/Engine/Plugins/Media/PixelStreaming2/Source/PixelStreaming2RTC/Internal/EpicRtcAudioCapturer.h
--- a/Engine/Plugins/Media/PixelStreaming2/Source/PixelStreaming2RTC/Internal/EpicRtcAudioCapturer.h
+++ b/Engine/Plugins/Media/PixelStreaming2/Source/PixelStreaming2RTC/Internal/EpicRtcAudioCapturer.h
@@ -15,17 +15,20 @@ namespace UE::PixelStreaming2
 	class FEpicRtcAudioCapturer : public FAudioCapturer
 	{
 	public:
-		static UE_API TSharedPtr<FEpicRtcAudioCapturer> Create();
 		virtual ~FEpicRtcAudioCapturer() = default;
 
 		// Override the push audio method as EpicRtc needs the broadcasted audio to be in 10ms chunks
 		UE_API virtual void PushAudio(const float* AudioData, int32 NumSamples, int32 NumChannels, int32 SampleRate) override;
 
 	protected:
-		FEpicRtcAudioCapturer() = default;
+		UE_API FEpicRtcAudioCapturer(const int SampleRate = 48000, const int NumChannels = 2, const float SampleSizeInSeconds = 0.5f)
+			: FAudioCapturer(SampleRate, NumChannels, SampleSizeInSeconds) {}
 
 	private:
 		TArray<int16_t> RecordingBuffer;
+
+		// Needed so FAudioCapturer::Create can access the constructor
+		friend FAudioCapturer;
 	};
 } // namespace UE::PixelStreaming2
 
diff --git a/Engine/Plugins/Media/PixelStreaming2/Source/PixelStreaming2RTC/Private/EpicRtcAudioCapturer.cpp b/Engine/Plugins/Media/PixelStreaming2/Source/PixelStreaming2RTC/Private/EpicRtcAudioCapturer.cpp
--- a/Engine/Plugins/Media/PixelStreaming2/Source/PixelStreaming2RTC/Private/EpicRtcAudioCapturer.cpp
+++ b/Engine/Plugins/Media/PixelStreaming2/Source/PixelStreaming2RTC/Private/EpicRtcAudioCapturer.cpp
@@ -8,31 +8,13 @@
 
 namespace UE::PixelStreaming2
 {
-	TSharedPtr<FEpicRtcAudioCapturer> FEpicRtcAudioCapturer::Create()
-	{
-		TSharedPtr<FEpicRtcAudioCapturer> AudioMixingCapturer(new FEpicRtcAudioCapturer());
-
-		FAudioDeviceManagerDelegates::OnAudioDeviceCreated.AddSP(AudioMixingCapturer.ToSharedRef(), &FEpicRtcAudioCapturer::CreateAudioProducer);
-		FAudioDeviceManagerDelegates::OnAudioDeviceDestroyed.AddSP(AudioMixingCapturer.ToSharedRef(), &FEpicRtcAudioCapturer::RemoveAudioProducer);
-
-		if (UPixelStreaming2PluginSettings::FDelegates* Delegates = UPixelStreaming2PluginSettings::Delegates())
-		{
-			Delegates->OnDebugDumpAudioChanged.AddSP(AudioMixingCapturer.ToSharedRef(), &FEpicRtcAudioCapturer::OnDebugDumpAudioChanged);
-
-			TWeakPtr<FEpicRtcAudioCapturer> WeakAudioMixingCapturer = AudioMixingCapturer;
-			FCoreDelegates::OnEnginePreExit.AddLambda([WeakAudioMixingCapturer]() {
-				if (TSharedPtr<FEpicRtcAudioCapturer> AudioMixingCapturer = WeakAudioMixingCapturer.Pin())
-				{
-					AudioMixingCapturer->OnEnginePreExit();
-				}
-			});
-		}
-
-		return AudioMixingCapturer;
-	}
-
 	void FEpicRtcAudioCapturer::PushAudio(const float* AudioData, int32 InNumSamples, int32 InNumChannels, const int32 InSampleRate)
 	{
+		check(AudioData != nullptr);
+		check(InNumSamples > 0);
+		check(InNumChannels > 0);
+		check(InSampleRate > 0);
+		check(InSampleRate == SampleRate);
 		Audio::TSampleBuffer<int16_t> Buffer(AudioData, InNumSamples, InNumChannels, SampleRate);
 		RecordingBuffer.Append(Buffer.GetData(), Buffer.GetNumSamples());
 
@@ -40,7 +22,7 @@ namespace UE::PixelStreaming2
 		const size_t BytesPerFrame = NumChannels * sizeof(int16_t);
 
 		// Feed in 10ms chunks
-		while (RecordingBuffer.Num() > SamplesPer10Ms)
+		while (RecordingBuffer.Num() >= SamplesPer10Ms)
 		{
 			OnAudioBuffer.Broadcast(RecordingBuffer.GetData(), SamplesPer10Ms, NumChannels, SampleRate);
 
diff --git a/Engine/Plugins/Media/PixelStreaming2/Source/PixelStreaming2RTC/Private/EpicRtcStreamer.cpp b/Engine/Plugins/Media/PixelStreaming2/Source/PixelStreaming2RTC/Private/EpicRtcStreamer.cpp
--- a/Engine/Plugins/Media/PixelStreaming2/Source/PixelStreaming2RTC/Private/EpicRtcStreamer.cpp
+++ b/Engine/Plugins/Media/PixelStreaming2/Source/PixelStreaming2RTC/Private/EpicRtcStreamer.cpp
@@ -2,6 +2,7 @@
 
 #include "EpicRtcStreamer.h"
 
+#include "AudioCapturer.h"
 #include "PixelStreaming2PluginSettings.h"
 #include "PixelStreaming2Delegates.h"
 #include "Engine/GameEngine.h"
@@ -272,7 +273,7 @@ namespace UE::PixelStreaming2
 		VideoCapturer->ResetFrameCapturer();
 
 		TickableTasks = FPixelStreaming2RTCModule::GetModule()->GetSharedTickableTasks();
-		AudioCapturer = FEpicRtcAudioCapturer::Create();
+		AudioCapturer = FAudioCapturer::Create<FEpicRtcAudioCapturer>();
 
 		{
 			FScopeLock Lock(&CustomAudioProducersCS);
