[PS2] Export virtual destructor on Audio classes.

diff --git a/Engine/Plugins/Media/PixelStreaming2/Source/PixelStreaming2/Internal/AudioCapturer.h b/Engine/Plugins/Media/PixelStreaming2/Source/PixelStreaming2/Internal/AudioCapturer.h
--- a/Engine/Plugins/Media/PixelStreaming2/Source/PixelStreaming2/Internal/AudioCapturer.h
+++ b/Engine/Plugins/Media/PixelStreaming2/Source/PixelStreaming2/Internal/AudioCapturer.h
@@ -20,7 +20,7 @@ namespace UE::PixelStreaming2
 	{
 	public:
 		UE_API FAudioPatchMixer(uint8 NumChannels, uint32 SampleRate, float SampleSizeSeconds);
-		virtual ~FAudioPatchMixer() = default;
+		UE_API virtual ~FAudioPatchMixer() = default;
 
 		UE_API uint32 GetMaxBufferSize() const;
 		UE_API uint8  GetNumChannels() const;
diff --git a/Engine/Plugins/Media/PixelStreaming2/Source/PixelStreaming2RTC/Internal/EpicRtcAudioCapturer.h b/Engine/Plugins/Media/PixelStreaming2/Source/PixelStreaming2RTC/Internal/EpicRtcAudioCapturer.h
--- a/Engine/Plugins/Media/PixelStreaming2/Source/PixelStreaming2RTC/Internal/EpicRtcAudioCapturer.h
+++ b/Engine/Plugins/Media/PixelStreaming2/Source/PixelStreaming2RTC/Internal/EpicRtcAudioCapturer.h
@@ -15,7 +15,7 @@ namespace UE::PixelStreaming2
 	class FEpicRtcAudioCapturer : public FAudioCapturer
 	{
 	public:
-		virtual ~FEpicRtcAudioCapturer() = default;
+		UE_API virtual ~FEpicRtcAudioCapturer() = default;
 
 		// Override the push audio method as EpicRtc needs the broadcasted audio to be in 10ms chunks
 		UE_API virtual void PushAudio(const float* AudioData, int32 NumSamples, int32 NumChannels, int32 SampleRate) override;
