[PS2] Fix: Input handler not unregistering

diff --git a/Engine/Plugins/Media/PixelStreaming2/Source/PixelStreaming2Input/Private/InputHandler.cpp b/Engine/Plugins/Media/PixelStreaming2/Source/PixelStreaming2Input/Private/InputHandler.cpp
--- a/Engine/Plugins/Media/PixelStreaming2/Source/PixelStreaming2Input/Private/InputHandler.cpp
+++ b/Engine/Plugins/Media/PixelStreaming2/Source/PixelStreaming2Input/Private/InputHandler.cpp
@@ -52,6 +52,8 @@ namespace UE::PixelStreaming2Input
 
 	FPixelStreaming2InputHandler::~FPixelStreaming2InputHandler()
 	{
+		IModularFeatures::Get().UnregisterModularFeature(GetModularFeatureName(), this);
+
 		if (UPixelStreaming2PluginSettings::FDelegates* Delegates = UPixelStreaming2PluginSettings::Delegates())
 		{
 			Delegates->OnInputKeyFilterChanged.Remove(OnInputKeyFilterChangedHandle);
