[PS2] Fix: Remove usage of <SWindow>.GetTitle().ToString() in FVideoProducerViewportBase

#Synced-CL 45701360

diff --git a/Engine/Plugins/Media/PixelStreaming2/Source/PixelStreaming2/Private/VideoProducerViewportBase.cpp b/Engine/Plugins/Media/PixelStreaming2/Source/PixelStreaming2/Private/VideoProducerViewportBase.cpp
--- a/Engine/Plugins/Media/PixelStreaming2/Source/PixelStreaming2/Private/VideoProducerViewportBase.cpp
+++ b/Engine/Plugins/Media/PixelStreaming2/Source/PixelStreaming2/Private/VideoProducerViewportBase.cpp
@@ -91,7 +91,7 @@ namespace UE::PixelStreaming2
 		TRefCountPtr<IPooledRenderTarget> OutTexture;
 		{
 			FRDGBuilder GraphBuilder(FRHICommandListImmediate::Get());
-			FRDGTextureRef InputTexture = GraphBuilder.RegisterExternalTexture(CreateRenderTarget(BackBuffer, *Window.GetTitle().ToString()));
+			FRDGTextureRef InputTexture = GraphBuilder.RegisterExternalTexture(CreateRenderTarget(BackBuffer, TEXT("VideoProducerViewportBaseInput")));
 			// TODO (william.belcher): The cropped texture having an 8bit pixel format matches the full editor video producer code path but we may need to reconsider this if we want to add 10bit video support
 			FRDGTextureRef CroppedTexture = GraphBuilder.CreateTexture(FRDGTextureDesc::Create2D(VirtualSize, EPixelFormat::PF_B8G8R8A8, FClearValueBinding::None, GetTexCreateFlags()), TEXT("VideoProducerViewportBaseCropped"));
 
diff --git a/Engine/Plugins/Media/PixelStreaming2/Source/PixelStreaming2Editor/Private/VideoProducerBackBufferComposited.cpp b/Engine/Plugins/Media/PixelStreaming2/Source/PixelStreaming2Editor/Private/VideoProducerBackBufferComposited.cpp
--- a/Engine/Plugins/Media/PixelStreaming2/Source/PixelStreaming2Editor/Private/VideoProducerBackBufferComposited.cpp
+++ b/Engine/Plugins/Media/PixelStreaming2/Source/PixelStreaming2Editor/Private/VideoProducerBackBufferComposited.cpp
@@ -102,8 +102,6 @@ namespace UE::EditorPixelStreaming2
 		 * Finally, the composited frame is extracted from the RDG pipeline and we send it on its way through the PixelCapturer
 		 *
 		 */
-		UE_LOG(LogPixelStreaming2BackBufferComposited, Verbose, TEXT("Type: %s"), *SlateWindow.GetTitle().ToString());
-
 		TSharedPtr<FVideoProducerUserData> UserData = MakeShared<FVideoProducerUserData>();
 		UserData->ProducerName = TEXT("FVideoProducerBackBufferComposited");
 		UserData->ProductionBeginCycles = FPlatformTime::Cycles64();
@@ -128,7 +126,7 @@ namespace UE::EditorPixelStreaming2
 				FRDGBuilder GraphBuilder(FRHICommandListImmediate::Get());
 
 				// Register an external RDG texture from the provided frame buffer
-				FRDGTextureRef InputTexture = GraphBuilder.RegisterExternalTexture(CreateRenderTarget(FrameBuffer, *SlateWindow.GetTitle().ToString()));
+				FRDGTextureRef InputTexture = GraphBuilder.RegisterExternalTexture(CreateRenderTarget(FrameBuffer, TEXT("VideoProducerBackBufferCompositedInput")));
 
 				// Create an internal RDG texture with the same extent and format as the source
 				FRDGTextureRef OutputTexture = GraphBuilder.CreateTexture(FRDGTextureDesc::Create2D(InputTexture->Desc.Extent, InputTexture->Desc.Format, FClearValueBinding::None, GetTexCreateFlags()), TEXT("VideoProducerBackBufferCompositedStaging"));
