[PS2] Fix: Crashes when PropertyChangedEvent.Property is nullptr

#Synced-CL 45701360

diff --git a/Engine/Plugins/Media/PixelStreaming2/Source/PixelStreaming2Settings/Private/PixelStreaming2PluginSettings.cpp b/Engine/Plugins/Media/PixelStreaming2/Source/PixelStreaming2Settings/Private/PixelStreaming2PluginSettings.cpp
--- a/Engine/Plugins/Media/PixelStreaming2/Source/PixelStreaming2Settings/Private/PixelStreaming2PluginSettings.cpp
+++ b/Engine/Plugins/Media/PixelStreaming2/Source/PixelStreaming2Settings/Private/PixelStreaming2PluginSettings.cpp
@@ -890,6 +890,11 @@ void UPixelStreaming2PluginSettings::PostEditChangeProperty(FPropertyChangedEven
 {
 	Super::PostEditChangeProperty(PropertyChangedEvent);
 
+	if (!PropertyChangedEvent.Property)
+	{
+		return;
+	}
+
 	FString PropertyName = PropertyChangedEvent.Property->GetNameCPP();
 
 	FString CVarName;
