[UCSPixelStreaming] Add: LogCensoring to restore existing functionality from FCensorLog (EOSPixelStreaming/EOSRtcManager.cpp)

diff --git a/Engine/Plugins/Media/PixelStreaming2/Source/PixelStreaming2RTC/Internal/EpicRtcLogging.h b/Engine/Plugins/Media/PixelStreaming2/Source/PixelStreaming2RTC/Internal/EpicRtcLogging.h
--- a/Engine/Plugins/Media/PixelStreaming2/Source/PixelStreaming2RTC/Internal/EpicRtcLogging.h
+++ b/Engine/Plugins/Media/PixelStreaming2/Source/PixelStreaming2RTC/Internal/EpicRtcLogging.h
@@ -61,8 +61,12 @@ namespace UE::PixelStreaming2
 	class FEpicRtcLogFilter : public ILogManipulator
 	{
 	public:
+		// Default constructor. Populates the contents of RegexPatterns with values from CVarEpicRtcLogFilter
 		UE_API FEpicRtcLogFilter();
-		UE_API ~FEpicRtcLogFilter();
+		// Alternate constructor that takes in a pre-populated array. Used for internal testing
+		UE_API FEpicRtcLogFilter(const TArray<FRegexPattern>& RegexPatterns);
+
+		UE_API virtual ~FEpicRtcLogFilter();
 
 		UE_API virtual bool IsFiltered(ELogVerbosity::Type LogVerbosity, const FString& LogString) override;
 
diff --git a/Engine/Plugins/Media/PixelStreaming2/Source/PixelStreaming2RTC/Private/EpicRtcLogging.cpp b/Engine/Plugins/Media/PixelStreaming2/Source/PixelStreaming2RTC/Private/EpicRtcLogging.cpp
--- a/Engine/Plugins/Media/PixelStreaming2/Source/PixelStreaming2RTC/Private/EpicRtcLogging.cpp
+++ b/Engine/Plugins/Media/PixelStreaming2/Source/PixelStreaming2RTC/Private/EpicRtcLogging.cpp
@@ -14,17 +14,22 @@ namespace UE::PixelStreaming2
 	{
 		ParseFilterString(UPixelStreaming2PluginSettings::CVarEpicRtcLogFilter.GetValueOnAnyThread());
 
-		if (UPixelStreaming2PluginSettings::FDelegates* Delegates = UPixelStreaming2PluginSettings::Delegates())
+		if (UPixelStreaming2PluginSettings::FDelegates* Delegates = UPixelStreaming2PluginSettings::Delegates(); Delegates)
 		{
 			EpicRtcLogFilterChangedHandle = Delegates->OnEpicRtcLogFilterChanged.AddRaw(this, &FEpicRtcLogFilter::OnEpicRtcLogFilterChanged);
 		}
 	}
 
+	FEpicRtcLogFilter::FEpicRtcLogFilter(const TArray<FRegexPattern>& RegexPatterns)
+		: RegexPatterns(RegexPatterns)
+	{
+	}
+
 	FEpicRtcLogFilter::~FEpicRtcLogFilter()
 	{
-		if (UPixelStreaming2PluginSettings::FDelegates* Delegates = UPixelStreaming2PluginSettings::Delegates())
+		if (UPixelStreaming2PluginSettings::FDelegates* Delegates = UPixelStreaming2PluginSettings::Delegates(); Delegates && EpicRtcLogFilterChangedHandle.IsValid())
 		{
-			Delegates->OnSimulcastEnabledChanged.Remove(EpicRtcLogFilterChangedHandle);
+			Delegates->OnEpicRtcLogFilterChanged.Remove(EpicRtcLogFilterChangedHandle);
 		}
 	}
 
