[PixelStreaming2] Fix: Non-unity Android compilation errors

diff --git a/Engine/Plugins/Media/PixelStreaming2/Source/PixelStreaming2RTC/Internal/EpicRtcPlatform.h b/Engine/Plugins/Media/PixelStreaming2/Source/PixelStreaming2RTC/Internal/EpicRtcPlatform.h
--- a/Engine/Plugins/Media/PixelStreaming2/Source/PixelStreaming2RTC/Internal/EpicRtcPlatform.h
+++ b/Engine/Plugins/Media/PixelStreaming2/Source/PixelStreaming2RTC/Internal/EpicRtcPlatform.h
@@ -6,8 +6,6 @@
 
 #include "epic_rtc/core/platform.h"
 
-#define UE_API PIXELSTREAMING2RTC_API
-
 #if PLATFORM_ANDROID
 	
     #include "Android/AndroidApplication.h"
@@ -20,12 +18,10 @@
     {
         class FEpicRtcAndroidPlatform : public EpicRtcAndroidPlatformInterface
         {
-            UE_API virtual JavaVM* GetJavaVM() override;
+            PIXELSTREAMING2RTC_API virtual JavaVM* GetJavaVM() override;
             
             virtual ~FEpicRtcAndroidPlatform() = default;
         };
     }
     
 #endif // PLATFORM_ANDROID
-
-#undef UE_API
\ No newline at end of file
