[PixelStreaming2] Add: Ability to access VideoConsumer in BP and retrieve FTextureResource* in C++

#Synced-CL 48443255

diff --git a/Engine/Plugins/Media/PixelStreaming2/Source/PixelStreaming2/Internal/Blueprints/PixelStreaming2MediaTexture.h b/Engine/Plugins/Media/PixelStreaming2/Source/PixelStreaming2/Internal/Blueprints/PixelStreaming2MediaTexture.h
--- a/Engine/Plugins/Media/PixelStreaming2/Source/PixelStreaming2/Internal/Blueprints/PixelStreaming2MediaTexture.h
+++ b/Engine/Plugins/Media/PixelStreaming2/Source/PixelStreaming2/Internal/Blueprints/PixelStreaming2MediaTexture.h
@@ -26,6 +26,8 @@ public:
 	// IPixelStreamingVideoConsumer implementation
 	UE_API void ConsumeFrame(FTextureRHIRef Frame);
 
+	UE_API const FTextureResource* GetResource() const;
+
 protected:
 	// UObject overrides.
 	UE_API virtual void BeginDestroy() override;
diff --git a/Engine/Plugins/Media/PixelStreaming2/Source/PixelStreaming2/Private/Blueprints/PixelStreaming2MediaTexture.cpp b/Engine/Plugins/Media/PixelStreaming2/Source/PixelStreaming2/Private/Blueprints/PixelStreaming2MediaTexture.cpp
--- a/Engine/Plugins/Media/PixelStreaming2/Source/PixelStreaming2/Private/Blueprints/PixelStreaming2MediaTexture.cpp
+++ b/Engine/Plugins/Media/PixelStreaming2/Source/PixelStreaming2/Private/Blueprints/PixelStreaming2MediaTexture.cpp
@@ -122,6 +122,11 @@ void UPixelStreaming2MediaTexture::ConsumeFrame(FTextureRHIRef Frame)
 	});
 }
 
+const FTextureResource* UPixelStreaming2MediaTexture::GetResource() const
+{
+	return CurrentResource;
+}
+
 void UPixelStreaming2MediaTexture::InitializeResources()
 {
 	ENQUEUE_RENDER_COMMAND(FPixelStreamingMediaTextureUpdateTextureReference)
diff --git a/Engine/Plugins/Media/PixelStreaming2/Source/PixelStreaming2RTC/Private/Blueprints/PixelStreaming2Peer.h b/Engine/Plugins/Media/PixelStreaming2/Source/PixelStreaming2RTC/Private/Blueprints/PixelStreaming2Peer.h
--- a/Engine/Plugins/Media/PixelStreaming2/Source/PixelStreaming2RTC/Private/Blueprints/PixelStreaming2Peer.h
+++ b/Engine/Plugins/Media/PixelStreaming2/Source/PixelStreaming2RTC/Private/Blueprints/PixelStreaming2Peer.h
@@ -93,7 +93,7 @@ public:
 	/**
 	 * A sink for the video data received once this connection has finished negotiating.
 	 */
-	UPROPERTY(EditDefaultsOnly, Category = "Properties", META = (DisplayName = "Pixel Streaming Video Consumer", AllowPrivateAccess = true))
+	UPROPERTY(EditDefaultsOnly, BlueprintReadOnly, Category = "Properties", META = (DisplayName = "Pixel Streaming Video Consumer", AllowPrivateAccess = true))
 	TObjectPtr<UPixelStreaming2MediaTexture> VideoConsumer = nullptr;
 
 private:
