[PixelStreaming2] Fix: Warnings about differing send and receive codec capabilities

#Synced-CL 48443255

diff --git a/Engine/Plugins/Media/PixelStreaming2/Source/PixelStreaming2RTC/Private/EpicRtcVideoDecoderInitializer.cpp b/Engine/Plugins/Media/PixelStreaming2/Source/PixelStreaming2RTC/Private/EpicRtcVideoDecoderInitializer.cpp
--- a/Engine/Plugins/Media/PixelStreaming2/Source/PixelStreaming2RTC/Private/EpicRtcVideoDecoderInitializer.cpp
+++ b/Engine/Plugins/Media/PixelStreaming2/Source/PixelStreaming2RTC/Private/EpicRtcVideoDecoderInitializer.cpp
@@ -133,15 +133,29 @@ namespace UE::PixelStreaming2
 		if (UE::PixelStreaming2::IsDecoderSupported<FVideoDecoderConfigVP8>())
 		{
 			Codecs[EVideoCodec::VP8].Add(TRefCountPtr<EpicRtcVideoCodecInfoInterface>(new FEpicRtcVideoCodecInfo(
-				EpicRtcVideoCodec::VP8
+				EpicRtcVideoCodec::VP8,
+				nullptr,
+				new FEpicRtcScalabilityModeArray({ EScalabilityMode::L1T1 })
 			)));
 		}
 
 		if (UE::PixelStreaming2::IsDecoderSupported<FVideoDecoderConfigVP9>())
 		{
+			using namespace UE::AVCodecCore::VP9;
 			Codecs[EVideoCodec::VP9].Add(TRefCountPtr<EpicRtcVideoCodecInfoInterface>(new FEpicRtcVideoCodecInfo(
-				EpicRtcVideoCodec::VP9
+				EpicRtcVideoCodec::VP9,
+				UE::PixelStreaming2::CreateVP9Format(EProfile::Profile0),
+				new FEpicRtcScalabilityModeArray(UE::PixelStreaming2::AllScalabilityModes)
 			)));
+
+			/**
+			 * Only advertise profile 0 until EpicRtc provides us with a way to extract the negotiated profile in the encoder config
+			 * 
+			 * Codecs[EVideoCodec::VP9].Add(TRefCountPtr<EpicRtcVideoCodecInfoInterface>(new FEpicRtcVideoCodecInfo(
+			 *	EpicRtcVideoCodec::VP9,
+			 *	UE::PixelStreaming2::CreateVP9Format(EProfile::Profile2),
+			 *	new FEpicRtcScalabilityModeArray(AllScalabilityModes))));
+			 */
 		}
 
 		if (UE::PixelStreaming2::IsDecoderSupported<FVideoDecoderConfigH264>())
@@ -149,18 +163,22 @@ namespace UE::PixelStreaming2
 			using namespace UE::AVCodecCore::H264;
 			Codecs[EVideoCodec::H264].Add(TRefCountPtr<EpicRtcVideoCodecInfoInterface>(new FEpicRtcVideoCodecInfo(
 				EpicRtcVideoCodec::H264,
-				UE::PixelStreaming2::CreateH264Format(EH264Profile::ConstrainedBaseline, EH264Level::Level_3_1)
+				UE::PixelStreaming2::CreateH264Format(EH264Profile::ConstrainedBaseline, EH264Level::Level_3_1),
+				new FEpicRtcScalabilityModeArray({ EScalabilityMode::L1T1 })
 			)));
 			Codecs[EVideoCodec::H264].Add(TRefCountPtr<EpicRtcVideoCodecInfoInterface>(new FEpicRtcVideoCodecInfo(
 				EpicRtcVideoCodec::H264,
-				UE::PixelStreaming2::CreateH264Format(EH264Profile::Baseline, EH264Level::Level_3_1)
+				UE::PixelStreaming2::CreateH264Format(EH264Profile::Baseline, EH264Level::Level_3_1),
+				new FEpicRtcScalabilityModeArray({ EScalabilityMode::L1T1 })
 			)));
 		}
 
 		if (UE::PixelStreaming2::IsDecoderSupported<FVideoDecoderConfigAV1>())
 		{
 			Codecs[EVideoCodec::AV1].Add(TRefCountPtr<EpicRtcVideoCodecInfoInterface>(new FEpicRtcVideoCodecInfo(
-				EpicRtcVideoCodec::AV1
+				EpicRtcVideoCodec::AV1,
+				nullptr,
+				new FEpicRtcScalabilityModeArray({ EScalabilityMode::L1T1 })
 			)));
 		}
 
