[PixelStreaming2] Fix: Audio processing pathways for resampling and channel mixing

#Synced-CL 48443255

diff --git a/Engine/Plugins/Media/PixelStreaming2/Source/PixelStreaming2/Internal/AudioCapturer.h b/Engine/Plugins/Media/PixelStreaming2/Source/PixelStreaming2/Internal/AudioCapturer.h
--- a/Engine/Plugins/Media/PixelStreaming2/Source/PixelStreaming2/Internal/AudioCapturer.h
+++ b/Engine/Plugins/Media/PixelStreaming2/Source/PixelStreaming2/Internal/AudioCapturer.h
@@ -16,42 +16,10 @@ namespace UE::PixelStreaming2
 {
 	class FAudioCapturer;
 
-	class FAudioPatchMixer : public Audio::FPatchMixer
-	{
-	public:
-		UE_API FAudioPatchMixer(uint8 NumChannels, uint32 SampleRate, float SampleSizeSeconds);
-		virtual ~FAudioPatchMixer() = default;
-
-		UE_API uint32 GetMaxBufferSize() const;
-		UE_API uint8  GetNumChannels() const;
-		UE_API uint32 GetSampleRate() const;
-
-	protected:
-		uint8  NumChannels;
-		uint32 SampleRate;
-		float  SampleSizeSeconds;
-	};
-
-	class FPatchInputProxy
-	{
-	public:
-		UE_API FPatchInputProxy(TSharedPtr<FAudioPatchMixer> InMixer);
-		UE_API virtual ~FPatchInputProxy();
-
-		UE_API void OnNewAudioFrame(const float* AudioData, int32 InNumSamples, int32 InNumChannels, int32 InSampleRate);
-
-	protected:
-		TSharedPtr<FAudioPatchMixer> Mixer;
-		Audio::FResampler			 Resampler;
-		Audio::FPatchInput			 PatchInput;
-		uint8						 NumChannels;
-		uint32						 SampleRate;
-	};
-
 	class FMixAudioTask : public FPixelStreamingTickableTask
 	{
 	public:
-		UE_API FMixAudioTask(TWeakPtr<FAudioCapturer> InCapturer, TSharedPtr<FAudioPatchMixer> InMixer);
+		UE_API FMixAudioTask(TSharedPtr<FAudioCapturer> Capturer);
 
 		virtual ~FMixAudioTask() = default;
 
@@ -61,24 +29,20 @@ namespace UE::PixelStreaming2
 		// End FPixelStreamingTickableTask
 
 	protected:
-		bool								  bIsRunning;
-		Audio::VectorOps::FAlignedFloatBuffer MixingBuffer;
-
-		TWeakPtr<FAudioCapturer>	 Capturer;
-		TSharedPtr<FAudioPatchMixer> Mixer;
+		TWeakPtr<FAudioCapturer> Capturer;
 	};
 
-	class FAudioCapturer : public TSharedFromThis<FAudioCapturer>
+	class FAudioCapturer : public Audio::FPatchMixer, public TSharedFromThis<FAudioCapturer>
 	{
 	public:
-		template<typename T>
-		static TSharedPtr<T> Create(const int InSampleRate = 48000, const int InNumChannels = 2, const float InSampleSizeInSeconds = 0.5f);
+		template <typename T>
+		static TSharedPtr<T> Create(const int32 InSampleRate = 48000, const int32 InNumChannels = 2, const float InSampleSizeInSeconds = 0.5f);
 
 		virtual ~FAudioCapturer() = default;
 
-		// Mixed audio input will push its audio to an FPatchInputProxy for mixing
-		UE_API void						  CreateAudioProducer(Audio::FDeviceId AudioDeviceId);
-		UE_API void						  RemoveAudioProducer(Audio::FDeviceId AudioDeviceId);
+		// AudioProducer lifecycle methods triggered by engine delegates
+		UE_API void CreateAudioProducer(Audio::FDeviceId AudioDeviceId);
+		UE_API void RemoveAudioProducer(Audio::FDeviceId AudioDeviceId);
 
 		// Methods for managing custom audio producers that will have their audio mixed in with engine audio
 		UE_API void AddAudioProducer(TSharedPtr<IPixelStreaming2AudioProducer> AudioProducer);
@@ -93,9 +57,13 @@ namespace UE::PixelStreaming2
 		DECLARE_TS_MULTICAST_DELEGATE_FourParams(FOnAudioBuffer, const int16_t*, int32, int32, const int32);
 		FOnAudioBuffer OnAudioBuffer;
 
+		/* */
+		void ProcessAudio();
+		/* */
+
 	protected:
-		UE_API FAudioCapturer(const int SampleRate = 48000, const int NumChannels = 2, const float SampleSizeInSeconds = 0.5f);
-		
+		UE_API FAudioCapturer(const int32 SampleRate = 48000, const int32 NumChannels = 2, const float SampleSizeInSeconds = 0.5f);
+
 		// Required to be called by derived classes if the class is not constructed with FAudioCapturer::Create.
 		// It initializes members that require TSharedPtr to FAudioCapturer
 		UE_API void Initialize();
@@ -103,29 +71,30 @@ namespace UE::PixelStreaming2
 		UE_API void OnDebugDumpAudioChanged(IConsoleVariable* Var);
 		UE_API void OnEnginePreExit();
 		UE_API void WriteDebugAudio();
-		// Allow the mix audio task to call OnAudio.
-		friend class FMixAudioTask;
-		UE_API void OnAudio(const float* AudioData, int32 NumSamples, int32 NumChannels, int32 SampleRate);
+
+		/* */
+		UE_API int32 GetSampleRate() const;
+		UE_API int32  GetNumChannels() const;
+		UE_API int32 GetMaxBufferSize() const;
+		/* */
 
 	protected:
-		TSharedPtr<FAudioPatchMixer> Mixer;
-		TSharedPtr<FMixAudioTask>	 MixerTask;
+		TSharedPtr<FMixAudioTask> MixerTask;
 
 		// Audio producers created by the engine that capture audio from specific devices
 		TMap<Audio::FDeviceId, TSharedPtr<FAudioProducer>> EngineAudioProducers;
 		// Audio producers created by the user that push audio from arbitrary sources
 		TMap<TSharedPtr<IPixelStreaming2AudioProducer>, TSharedPtr<FAudioProducer>> UserAudioProducers;
 
-		int				  SampleRate;
-		int				  NumChannels;
-		float			  SampleSizeSeconds;
-		Audio::FResampler Resampler;
+		int32	  SampleRate;
+		int32	  NumChannels;
+		float SampleSizeSeconds;
 
 		Audio::TSampleBuffer<int16_t> DebugDumpAudioBuffer;
 	};
 
 	template <typename T>
-	TSharedPtr<T> FAudioCapturer::Create(const int InSampleRate, const int InNumChannels, const float InSampleSizeInSeconds)
+	TSharedPtr<T> FAudioCapturer::Create(const int32 InSampleRate, const int32 InNumChannels, const float InSampleSizeInSeconds)
 	{
 		static_assert(std::is_base_of_v<FAudioCapturer, T>);
 		TSharedPtr<T> AudioCapturer(new T(InSampleRate, InNumChannels, InSampleSizeInSeconds));
diff --git a/Engine/Plugins/Media/PixelStreaming2/Source/PixelStreaming2/Internal/AudioProducer.h b/Engine/Plugins/Media/PixelStreaming2/Source/PixelStreaming2/Internal/AudioProducer.h
--- a/Engine/Plugins/Media/PixelStreaming2/Source/PixelStreaming2/Internal/AudioProducer.h
+++ b/Engine/Plugins/Media/PixelStreaming2/Source/PixelStreaming2/Internal/AudioProducer.h
@@ -3,24 +3,26 @@
 #pragma once
 
 #include "Audio.h"
-#include "ISubmixBufferListener.h"
+#include "DSP/MultithreadedPatching.h"
 #include "IPixelStreaming2AudioProducer.h"
+#include "ISubmixBufferListener.h"
+#include "UtilsAudio.h"
 
 #define UE_API PIXELSTREAMING2_API
 
 namespace UE::PixelStreaming2
 {
-	class FPatchInputProxy;
+	class FAudioCapturer;
 
 	/**
 	 * An audio input capable of listening to UE submix's as well as receiving user audio via the OnPushedAudio method.
-	 * Any received audio will be passed to the PatchInput's OnNewAudioFrame method
+	 * Any received audio will be pushed into the FPatchInput interface for mixing
 	 */
-	class FAudioProducer : public ISubmixBufferListener
+	class FAudioProducer : public ISubmixBufferListener, public Audio::FPatchInput
 	{
 	public:
-		static UE_API TSharedPtr<FAudioProducer> Create(Audio::FDeviceId AudioDeviceId, TSharedPtr<FPatchInputProxy> PatchInput);
-		static UE_API TSharedPtr<FAudioProducer> Create(TSharedPtr<IPixelStreaming2AudioProducer> AudioProducer, TSharedPtr<FPatchInputProxy> PatchInput);
+		static UE_API TSharedPtr<FAudioProducer> Create(Audio::FDeviceId AudioDeviceId, int32 TargetSampleRate, int32 TargetNumChannels, const Audio::FPatchOutputStrongPtr& PatchOutput);
+		static UE_API TSharedPtr<FAudioProducer> Create(TSharedPtr<IPixelStreaming2AudioProducer> AudioProducer, int32 TargetSampleRate, int32 TargetNumChannels, const Audio::FPatchOutputStrongPtr& PatchOutput);
 		virtual ~FAudioProducer() = default;
 
 		// ISubmixBufferListener interface
@@ -29,14 +31,16 @@ namespace UE::PixelStreaming2
 		// Listener for audio pushed from custom IPixelStreaming2AudioProducer implementations
 		UE_API void OnPushedAudio(const float* AudioData, int32 NumSamples, int32 NumChannels, int32 SampleRate);
 
-		void ToggleMuted() { bIsMuted = !bIsMuted; }
-
 	protected:
-		UE_API FAudioProducer(TSharedPtr<FPatchInputProxy> PatchInput);
+		UE_API FAudioProducer(int32 TargetSampleRate, int32 TargetNumChannels, const Audio::FPatchOutputStrongPtr& PatchOutput);
 
 	protected:
-		TSharedPtr<FPatchInputProxy> PatchInput;
-		bool						 bIsMuted = false;
+		const int32 TargetSampleRate;
+		const int32 TargetNumChannels;
+		bool bIsMuted;
+
+	private:
+		void OnRawAudio(const float* AudioData, int32 NumSamples, int32 NumChannels, int32 SampleRate);
 	};
 } // namespace UE::PixelStreaming2
 
diff --git a/Engine/Plugins/Media/PixelStreaming2/Source/PixelStreaming2/Internal/SoundGenerator.h b/Engine/Plugins/Media/PixelStreaming2/Source/PixelStreaming2/Internal/SoundGenerator.h
--- a/Engine/Plugins/Media/PixelStreaming2/Source/PixelStreaming2/Internal/SoundGenerator.h
+++ b/Engine/Plugins/Media/PixelStreaming2/Source/PixelStreaming2/Internal/SoundGenerator.h
@@ -26,27 +26,27 @@ namespace UE::PixelStreaming2
 		UE_API virtual int32 GetDesiredNumSamplesToRenderPerCallback() const;
 
 		// Optional. Called on audio generator thread right when the generator begins generating.
-		virtual void OnBeginGenerate() { bGeneratingAudio = true; };
+		UE_API virtual void OnBeginGenerate();
 
 		// Optional. Called on audio generator thread right when the generator ends generating.
-		virtual void OnEndGenerate() { bGeneratingAudio = false; };
+		UE_API virtual void OnEndGenerate();
 
 		// Optional. Can be overridden to end the sound when generating is finished.
-		virtual bool IsFinished() const { return false; };
+		UE_API virtual bool IsFinished() const;
 
-		UE_API void AddAudio(const int16_t* AudioData, int InSampleRate, size_t NChannels, size_t NFrames);
+		UE_API void AddAudio(const int16* AudioData, int SampleRate, size_t NumChannels, size_t NumFrames);
 
-		int32 GetSampleRate() { return Params.SampleRate; }
-		int32 GetNumChannels() { return Params.NumChannels; }
+		UE_API bool IsInitialized();
 		UE_API void  EmptyBuffers();
-		UE_API void  SetParameters(const FSoundGeneratorInitParams& InitParams);
+		UE_API void  Initialize(const FSoundGeneratorInitParams& InitParams);
 
 	private:
 		FSoundGeneratorInitParams Params;
-		TArray<int16_t>			  Buffer;
+		TArray<float>			  Buffer;
 		FCriticalSection		  CriticalSection;
 
 	public:
+		std::atomic<bool> bInitialized = false;
 		std::atomic<bool> bGeneratingAudio = false;
 		std::atomic<bool> bShouldGenerateAudio = false;
 	};
diff --git a/Engine/Plugins/Media/PixelStreaming2/Source/PixelStreaming2/Internal/UtilsAudio.h b/Engine/Plugins/Media/PixelStreaming2/Source/PixelStreaming2/Internal/UtilsAudio.h
new file mode 100644
--- /dev/null
+++ b/Engine/Plugins/Media/PixelStreaming2/Source/PixelStreaming2/Internal/UtilsAudio.h
@@ -0,0 +1,132 @@
+// Copyright Epic Games, Inc. All Rights Reserved.
+
+#pragma once
+
+#include "AudioResampler.h"
+#include "DSP/AlignedBuffer.h"
+#include "DSP/ConvertDeinterleave.h"
+#include "DSP/FloatArrayMath.h"
+
+#define UE_API PIXELSTREAMING2_API
+
+namespace UE::PixelStreaming2
+{
+	template <typename T>
+	struct TRawAudio
+	{
+		T*	  Data;
+		int32 NumSamples;
+        int32 SampleRate;
+		int32 NumChannels;
+        float Gain = 1.f;
+	};
+
+	struct FProcessedAudio
+	{
+		Audio::FAlignedFloatBuffer& Data;
+        int32 NumSamples;
+        int32 SampleRate;
+		int32 NumChannels;
+        float Gain = 1.f;
+	};
+
+	template <typename T>
+	bool ProcessAudio(const TRawAudio<T>& RawAudio, FProcessedAudio& OutProcessedAudio)
+	{
+		static_assert(std::is_same_v<std::decay_t<T>, int16> || std::is_same_v<std::decay_t<T>, float>, "Unsupported audio sample type");
+
+        if (RawAudio.Data == nullptr || RawAudio.NumChannels <= 0 || RawAudio.SampleRate <= 0 || RawAudio.NumSamples <= 0 || RawAudio.Gain <= 0.f)
+        {
+            return false;
+        }
+
+        if (OutProcessedAudio.SampleRate <= 0 || OutProcessedAudio.NumChannels <= 0 || OutProcessedAudio.Gain < 0.f)
+        {
+            return false;
+        }
+
+		TArray<float> RawAudioData;
+		RawAudioData.SetNumUninitialized(RawAudio.NumSamples);
+		if constexpr (std::is_same_v<std::decay_t<T>, int16>)
+		{
+			Audio::ArrayPcm16ToFloat(MakeArrayView(RawAudio.Data, RawAudio.NumSamples), RawAudioData);
+		}
+		else
+		{
+			FMemory::Memcpy(RawAudioData.GetData(), RawAudio.Data, RawAudio.NumSamples * sizeof(float));
+		}
+
+		if (RawAudio.SampleRate == OutProcessedAudio.SampleRate && RawAudio.NumChannels == OutProcessedAudio.NumChannels && RawAudio.Gain == OutProcessedAudio.Gain)
+		{
+			// No processing needed, just copy the data
+            int32 NumOutSamples = RawAudioData.Num();
+            OutProcessedAudio.Data.SetNumUninitialized(NumOutSamples);
+            FMemory::Memcpy(OutProcessedAudio.Data.GetData(), RawAudioData.GetData(), NumOutSamples * sizeof(float));
+			OutProcessedAudio.NumSamples = NumOutSamples;
+			return true;
+		}
+
+		TUniquePtr<Audio::IConvertDeinterleave> Converter = Audio::IConvertDeinterleave::Create({ .NumInputChannels = RawAudio.NumChannels, .NumOutputChannels = OutProcessedAudio.NumChannels });
+		if (!Converter)
+		{
+			return false;
+		}
+
+		Audio::FMultichannelBuffer MultichannelBuffer;
+		Audio::SetMultichannelBufferSize(OutProcessedAudio.NumChannels, RawAudio.NumSamples / RawAudio.NumChannels, MultichannelBuffer);
+		// ProcessAudio also handles the up/down mixing of channels
+		Converter->ProcessAudio(RawAudioData, MultichannelBuffer);
+
+		for (int32 ChannelIndex = 0; ChannelIndex < OutProcessedAudio.NumChannels; ChannelIndex++)
+		{
+			// Operate on the channels individually
+			Audio::FAlignedFloatBuffer& ChannelBuffer = MultichannelBuffer[ChannelIndex];
+
+			// Resample if required
+			if (RawAudio.SampleRate != OutProcessedAudio.SampleRate)
+			{
+				Audio::FResamplingParameters ResamplerParameters = {
+					.ResamplerMethod = Audio::EResamplingMethod::Linear,
+					.NumChannels = 1,
+					.SourceSampleRate = static_cast<float>(RawAudio.SampleRate),
+					.DestinationSampleRate = static_cast<float>(OutProcessedAudio.SampleRate),
+					.InputBuffer = ChannelBuffer
+				};
+
+				Audio::FAlignedFloatBuffer ResamplerOutputData;
+				ResamplerOutputData.AddUninitialized(Audio::GetOutputBufferSize(ResamplerParameters));
+
+				Audio::FResamplerResults ResamplerResults;
+				ResamplerResults.OutBuffer = &ResamplerOutputData;
+
+				if (!Audio::Resample(ResamplerParameters, ResamplerResults))
+				{
+					return false;
+				}
+
+				int32 NumSamplesGenerated = ResamplerResults.OutputFramesGenerated;
+				ChannelBuffer.SetNumZeroed(NumSamplesGenerated);
+				FMemory::Memcpy(ChannelBuffer.GetData(), ResamplerOutputData.GetData(), NumSamplesGenerated * sizeof(float));
+			}
+
+            if (RawAudio.Gain != OutProcessedAudio.Gain)
+			{
+				// This approach does loop over the channel buffer twice but can take advantage of SIMD. May be quicker?
+				Audio::ArrayMultiplyByConstantInPlace(ChannelBuffer, OutProcessedAudio.Gain / RawAudio.Gain);
+				Audio::ArrayClampInPlace(ChannelBuffer, -1.f, 1.f);
+			}
+		}
+
+		Audio::FAlignedFloatBuffer InterleavedBuffer;
+		Audio::ArrayInterleave(MultichannelBuffer, InterleavedBuffer);
+
+        int32 NumOutSamples = InterleavedBuffer.Num();
+        OutProcessedAudio.Data.SetNumUninitialized(NumOutSamples);
+        FMemory::Memcpy(OutProcessedAudio.Data.GetData(), InterleavedBuffer.GetData(), NumOutSamples * sizeof(float));
+		OutProcessedAudio.NumSamples = NumOutSamples;
+
+		return true;
+	}
+} // namespace UE::PixelStreaming2
+
+#undef UE_API
\ No newline at end of file
diff --git a/Engine/Plugins/Media/PixelStreaming2/Source/PixelStreaming2/Private/AudioCapturer.cpp b/Engine/Plugins/Media/PixelStreaming2/Source/PixelStreaming2/Private/AudioCapturer.cpp
--- a/Engine/Plugins/Media/PixelStreaming2/Source/PixelStreaming2/Private/AudioCapturer.cpp
+++ b/Engine/Plugins/Media/PixelStreaming2/Source/PixelStreaming2/Private/AudioCapturer.cpp
@@ -11,160 +11,25 @@
 
 namespace UE::PixelStreaming2
 {
-	FAudioPatchMixer::FAudioPatchMixer(uint8 NumChannels, uint32 SampleRate, float SampleSizeSeconds)
-		: NumChannels(NumChannels)
-		, SampleRate(SampleRate)
-		, SampleSizeSeconds(SampleSizeSeconds)
-	{
-	}
-
-	uint32 FAudioPatchMixer::GetMaxBufferSize() const
-	{
-		return NumChannels * SampleRate * SampleSizeSeconds;
-	}
-
-	uint8 FAudioPatchMixer::GetNumChannels() const
-	{
-		return NumChannels;
-	}
-
-	uint32 FAudioPatchMixer::GetSampleRate() const
-	{
-		return SampleRate;
-	}
-
 	/***************************************************
 	 *
 	 ****************************************************/
-	FPatchInputProxy::FPatchInputProxy(TSharedPtr<FAudioPatchMixer> InMixer)
-		: Mixer(InMixer)
-		// We don't want the patch input to handle gain as the capturer handles that at the end
-		, PatchInput(Mixer->AddNewInput(Mixer->GetMaxBufferSize(), 1.0f))
-		, NumChannels(Mixer->GetNumChannels())
-		, SampleRate(Mixer->GetSampleRate())
+	FMixAudioTask::FMixAudioTask(TSharedPtr<FAudioCapturer> Capturer)
+		: Capturer(Capturer)
 	{
 	}
 
-	FPatchInputProxy::~FPatchInputProxy()
-	{
-		Mixer->RemovePatch(PatchInput);
-	}
-
-	void FPatchInputProxy::OnNewAudioFrame(const float* AudioData, int32 InNumSamples, int32 InNumChannels, int32 InSampleRate)
-	{
-		TArray<float> AudioBuffer;
-
-		if (SampleRate != InSampleRate)
-		{
-			float SampleRateConversionRatio = static_cast<float>(SampleRate) / static_cast<float>(InSampleRate);
-			Resampler.Init(Audio::EResamplingMethod::Linear, SampleRateConversionRatio, InNumChannels);
-
-			int32 NumOriginalSamples = InNumSamples / InNumChannels;
-			int32 NumConvertedSamples = FMath::CeilToInt(InNumSamples / InNumChannels * SampleRateConversionRatio);
-			int32 OutputSamples = INDEX_NONE;
-			AudioBuffer.AddZeroed(NumConvertedSamples * InNumChannels);
-
-			// Perform the sample rate conversion
-			int32 ErrorCode = Resampler.ProcessAudio(const_cast<float*>(AudioData), NumOriginalSamples, false, AudioBuffer.GetData(), NumConvertedSamples, OutputSamples);
-			verifyf(OutputSamples <= NumConvertedSamples, TEXT("OutputSamples > NumConvertedSamples"));
-			if (ErrorCode != 0)
-			{
-				UE_LOG(LogPixelStreaming2, Warning, TEXT("(FAudioInput) Problem occured resampling audio data. Code: %d"), ErrorCode);
-				return;
-			}
-		}
-		else
-		{
-			AudioBuffer.AddZeroed(InNumSamples);
-			FMemory::Memcpy(AudioBuffer.GetData(), AudioData, AudioBuffer.Num() * sizeof(float));
-		}
-
-		// Note: TSampleBuffer takes in AudioData as float* and internally converts to int16
-		Audio::TSampleBuffer<int16_t> Buffer(AudioBuffer.GetData(), AudioBuffer.Num(), InNumChannels, SampleRate);
-		if (InNumChannels != NumChannels)
-		{
-			if (InNumChannels < NumChannels)
-			{
-				// Up mix by duplicating the mono source to each channel
-				TArrayView<int16> SourceBuffer = Buffer.GetArrayView();
-				TArray<int16>	  MixedBuffer;
-				MixedBuffer.SetNumZeroed(Buffer.GetNumSamples() / InNumChannels * NumChannels);
-				for (int32 SrcSampleIdx = 0; SrcSampleIdx < Buffer.GetNumSamples(); SrcSampleIdx++)
-				{
-					int32 DestSampleIdx = SrcSampleIdx * NumChannels;
-					for (int32 Channel = 0; Channel < NumChannels; Channel++)
-					{
-						MixedBuffer[DestSampleIdx + Channel] = SourceBuffer[SrcSampleIdx];
-					}
-				}
-
-				Buffer.CopyFrom(MixedBuffer, NumChannels, SampleRate);
-			}
-			else
-			{
-				// Down mix using inbuilt method
-				Buffer.MixBufferToChannels(NumChannels);
-			}
-		}
-
-		// Apply gain
-		float Gain = UPixelStreaming2PluginSettings::CVarWebRTCAudioGain.GetValueOnAnyThread();
-		if (Gain != 1.0f)
-		{
-			int16* PCMAudio = Buffer.GetArrayView().GetData();
-
-			// multiply audio by gain multiplier
-			for (int i = 0; i < Buffer.GetNumSamples(); i++)
-			{
-				*PCMAudio = FMath::Max(-32768, FMath::Min(32767, *PCMAudio * Gain));
-				PCMAudio++;
-			}
-		}
-
-		Audio::TSampleBuffer<float> BufferToPush(Buffer.GetData(), Buffer.GetNumSamples(), Buffer.GetNumChannels(), Buffer.GetSampleRate());
-		PatchInput.PushAudio(BufferToPush.GetData(), BufferToPush.GetNumSamples());
-	}
-
-	/***************************************************
-	 *
-	 ****************************************************/
-	FMixAudioTask::FMixAudioTask(TWeakPtr<FAudioCapturer> InCapturer, TSharedPtr<FAudioPatchMixer> InMixer)
-		: Capturer(InCapturer)
-		, Mixer(InMixer)
-	{
-		MixingBuffer.SetNumUninitialized(Mixer->GetMaxBufferSize());
-	}
-
 	void FMixAudioTask::Tick(float DeltaMs)
 	{
-		if (!Mixer)
-		{
-			return;
-		}
-
-		// 4 samples is the absolute minimum required for mixing
-		if (MixingBuffer.Num() < 4)
+		if (!Capturer.IsValid())
 		{
 			return;
 		}
 
-		int32 TargetNumSamples = Mixer->MaxNumberOfSamplesThatCanBePopped();
-		if (TargetNumSamples < 0)
-		{
-			return;
-		}
-
-		int32 NSamplesPopped = Mixer->PopAudio(MixingBuffer.GetData(), TargetNumSamples, false /* bUseLatestAudio */);
-		if (NSamplesPopped == 0)
-		{
-			return;
-		}
-
-		if (TSharedPtr<FAudioCapturer> PinnedCapturer = Capturer.Pin())
-		{
-			PinnedCapturer->OnAudio(MixingBuffer.GetData(), NSamplesPopped, Mixer->GetNumChannels(), Mixer->GetSampleRate());
-		}
+		TSharedPtr<FAudioCapturer> PinnedCapturer = Capturer.Pin();
+		PinnedCapturer->ProcessAudio();
 	}
+
 	const FString& FMixAudioTask::GetName() const
 	{
 		static FString TaskName = TEXT("MixAudioTask");
@@ -181,8 +46,7 @@ namespace UE::PixelStreaming2
 		if (GEngine)
 		{
 			// subscribe to audio data
-			Mixer = MakeShared<FAudioPatchMixer>(NumChannels, SampleRate, SampleSizeSeconds);
-			MixerTask = FPixelStreamingTickableTask::Create<FMixAudioTask>(AsWeak(), Mixer);
+			MixerTask = FPixelStreamingTickableTask::Create<FMixAudioTask>(AsShared());
 			if (FAudioDeviceManager* AudioDeviceManager = FAudioDeviceManager::Get(); AudioDeviceManager != nullptr)
 			{
 				TWeakPtr<FAudioCapturer> LocalAudioCapturer = AsWeak();
@@ -196,51 +60,55 @@ namespace UE::PixelStreaming2
 		}
 	}
 
-	FAudioCapturer::FAudioCapturer(const int SampleRate, const int NumChannels, const float NumSampleSizeInSeconds)
+	FAudioCapturer::FAudioCapturer(const int32 SampleRate, const int32 NumChannels, const float SampleSizeInSeconds)
 		: SampleRate(SampleRate)
 		, NumChannels(NumChannels)
-		, SampleSizeSeconds(NumSampleSizeInSeconds)
+		, SampleSizeSeconds(SampleSizeInSeconds)
 	{
 	}
 
 	void FAudioCapturer::CreateAudioProducer(Audio::FDeviceId AudioDeviceId)
 	{
 		// The lifetimes of audio producers created by the engine are our responsibility
-		TSharedPtr<FAudioProducer> AudioInput = FAudioProducer::Create(AudioDeviceId, MakeShared<FPatchInputProxy>(Mixer));
+		TSharedPtr<FAudioProducer> AudioInput = FAudioProducer::Create(AudioDeviceId, GetSampleRate(), GetNumChannels(), MakeShared<Audio::FPatchOutput, ESPMode::ThreadSafe>(GetMaxBufferSize()));
 		EngineAudioProducers.Add(AudioDeviceId, AudioInput);
+
+		AddNewInput(*AudioInput.Get());
 	}
 
 	void FAudioCapturer::RemoveAudioProducer(Audio::FDeviceId AudioDeviceId)
 	{
-		EngineAudioProducers.Remove(AudioDeviceId);
+		TSharedPtr<FAudioProducer> AudioInput;
+		if (!EngineAudioProducers.RemoveAndCopyValue(AudioDeviceId, AudioInput) || !AudioInput.IsValid())
+		{
+			return;
+		}
+
+		RemovePatch(*AudioInput.Get());
 	}
 
 	void FAudioCapturer::AddAudioProducer(TSharedPtr<IPixelStreaming2AudioProducer> AudioProducer)
 	{
-		TSharedPtr<FAudioProducer> AudioInput = FAudioProducer::Create(AudioProducer, MakeShared<FPatchInputProxy>(Mixer));
+		TSharedPtr<FAudioProducer> AudioInput = FAudioProducer::Create(AudioProducer, GetSampleRate(), GetNumChannels(), MakeShared<Audio::FPatchOutput, ESPMode::ThreadSafe>(GetMaxBufferSize()));
 		UserAudioProducers.Add(AudioProducer, AudioInput);
-	}
 
-	void FAudioCapturer::RemoveAudioProducer(TSharedPtr<IPixelStreaming2AudioProducer> AudioProducer)
-	{
-		UserAudioProducers.Remove(AudioProducer);
+		AddNewInput(*AudioInput.Get());
 	}
 
-	void FAudioCapturer::OnAudio(const float* AudioData, int32 InNumSamples, int32 InNumChannels, int32 InSampleRate)
+	void FAudioCapturer::RemoveAudioProducer(TSharedPtr<IPixelStreaming2AudioProducer> AudioProducer)
 	{
-		if (UPixelStreaming2PluginSettings::CVarDebugDumpAudio.GetValueOnAnyThread())
+		TSharedPtr<FAudioProducer> AudioInput;
+		if (!UserAudioProducers.RemoveAndCopyValue(AudioProducer, AudioInput) || !AudioInput.IsValid())
 		{
-			// Note: TSampleBuffer takes in AudioData as float* and internally converts to int16
-			Audio::TSampleBuffer<int16_t> Buffer(AudioData, InNumSamples, InNumChannels, SampleRate);
-			DebugDumpAudioBuffer.Append(Buffer.GetData(), Buffer.GetNumSamples(), Buffer.GetNumChannels(), Buffer.GetSampleRate());
+			return;
 		}
 
-		PushAudio(AudioData, InNumSamples, InNumChannels, InSampleRate);
+		RemovePatch(*AudioInput.Get());
 	}
 
 	void FAudioCapturer::PushAudio(const float* AudioData, int32 InNumSamples, int32 InNumChannels, int32 InSampleRate)
 	{
-		Audio::TSampleBuffer<int16_t> Buffer(AudioData, InNumSamples, InNumChannels, SampleRate);
+		Audio::TSampleBuffer<int16_t> Buffer(AudioData, InNumSamples, InNumChannels, InSampleRate);
 
 		OnAudioBuffer.Broadcast(Buffer.GetData(), Buffer.GetNumSamples(), Buffer.GetNumChannels(), Buffer.GetSampleRate());
 	}
@@ -276,4 +144,46 @@ namespace UE::PixelStreaming2
 		UE_LOGFMT(LogPixelStreaming2, Log, "Saving audio sample to: {0}", FilePath);
 		DebugDumpAudioBuffer.Reset();
 	}
+
+	int32 FAudioCapturer::GetSampleRate() const
+	{
+		return SampleRate;
+	}
+
+	int32 FAudioCapturer::GetNumChannels() const
+	{
+		return NumChannels;
+	}
+
+	int32 FAudioCapturer::GetMaxBufferSize() const
+	{
+		return NumChannels * SampleRate * SampleSizeSeconds;
+	}
+
+	void FAudioCapturer::ProcessAudio()
+	{
+		int32 TargetNumSamples = MaxNumberOfSamplesThatCanBePopped();
+		if (TargetNumSamples < 0)
+		{
+			return;
+		}
+
+		Audio::VectorOps::FAlignedFloatBuffer MixingBuffer;
+		MixingBuffer.SetNumUninitialized(GetMaxBufferSize());
+
+		int32 NumSamplesPopped = PopAudio(MixingBuffer.GetData(), TargetNumSamples, false /* bUseLatestAudio */);
+		if (NumSamplesPopped == 0)
+		{
+			return;
+		}
+
+		if (UPixelStreaming2PluginSettings::CVarDebugDumpAudio.GetValueOnAnyThread())
+		{
+			// Note: TSampleBuffer takes in AudioData as float* and internally converts to int16
+			Audio::TSampleBuffer<int16_t> Buffer(MixingBuffer.GetData(), NumSamplesPopped, GetNumChannels(), GetSampleRate());
+			DebugDumpAudioBuffer.Append(Buffer.GetData(), Buffer.GetNumSamples(), Buffer.GetNumChannels(), Buffer.GetSampleRate());
+		}
+
+		PushAudio(MixingBuffer.GetData(), NumSamplesPopped, GetNumChannels(), GetSampleRate());
+	}
 } // namespace UE::PixelStreaming2
\ No newline at end of file
diff --git a/Engine/Plugins/Media/PixelStreaming2/Source/PixelStreaming2/Private/AudioProducer.cpp b/Engine/Plugins/Media/PixelStreaming2/Source/PixelStreaming2/Private/AudioProducer.cpp
--- a/Engine/Plugins/Media/PixelStreaming2/Source/PixelStreaming2/Private/AudioProducer.cpp
+++ b/Engine/Plugins/Media/PixelStreaming2/Source/PixelStreaming2/Private/AudioProducer.cpp
@@ -5,12 +5,14 @@
 #include "AudioCapturer.h"
 #include "AudioDevice.h"
 #include "AudioDeviceManager.h"
+#include "Logging.h"
+#include "UtilsAudio.h"
 
 namespace UE::PixelStreaming2
 {
-	TSharedPtr<FAudioProducer> FAudioProducer::Create(Audio::FDeviceId InAudioDeviceId, TSharedPtr<FPatchInputProxy> InPatchInput)
+	TSharedPtr<FAudioProducer> FAudioProducer::Create(Audio::FDeviceId InAudioDeviceId, int32 TargetSampleRate, int32 TargetNumChannels, const Audio::FPatchOutputStrongPtr& PatchOutput)
 	{
-		TSharedPtr<FAudioProducer> Listener = TSharedPtr<FAudioProducer>(new FAudioProducer(InPatchInput));
+		TSharedPtr<FAudioProducer> Listener = TSharedPtr<FAudioProducer>(new FAudioProducer(TargetSampleRate, TargetNumChannels, PatchOutput));
 		if (FAudioDeviceManager* AudioDeviceManager = FAudioDeviceManager::Get(); AudioDeviceManager != nullptr)
 		{
 			if (FAudioDevice* AudioDevice = AudioDeviceManager->GetAudioDeviceRaw(InAudioDeviceId); AudioDevice != nullptr)
@@ -22,36 +24,61 @@ namespace UE::PixelStreaming2
 		return Listener;
 	}
 
-	TSharedPtr<FAudioProducer> FAudioProducer::Create(TSharedPtr<IPixelStreaming2AudioProducer> InAudioProducer, TSharedPtr<FPatchInputProxy> InPatchInput)
+	TSharedPtr<FAudioProducer> FAudioProducer::Create(TSharedPtr<IPixelStreaming2AudioProducer> InAudioProducer, int32 TargetSampleRate, int32 TargetNumChannels, const Audio::FPatchOutputStrongPtr& PatchOutput)
 	{
-		TSharedPtr<FAudioProducer> Listener = TSharedPtr<FAudioProducer>(new FAudioProducer(InPatchInput));
+		TSharedPtr<FAudioProducer> Listener = TSharedPtr<FAudioProducer>(new FAudioProducer(TargetSampleRate, TargetNumChannels, PatchOutput));
 		InAudioProducer->OnAudioPushed.AddSP(Listener.ToSharedRef(), &FAudioProducer::OnPushedAudio);
 
 		return Listener;
 	}
 
-	FAudioProducer::FAudioProducer(TSharedPtr<FPatchInputProxy> PatchInput)
-		: PatchInput(PatchInput)
+	FAudioProducer::FAudioProducer(int32 TargetSampleRate, int32 TargetNumChannels, const Audio::FPatchOutputStrongPtr& PatchOutput)
+		: FPatchInput(PatchOutput)
+		, TargetSampleRate(TargetSampleRate)
+		, TargetNumChannels(TargetNumChannels)
+		, bIsMuted(false)
 	{
 	}
 
 	void FAudioProducer::OnNewSubmixBuffer(const USoundSubmix* OwningSubmix, float* AudioData, int32 NumSamples, int32 NumChannels, const int32 SampleRate, double AudioClock)
+	{
+		OnRawAudio(AudioData, NumSamples, NumChannels, SampleRate);
+	}
+
+	void FAudioProducer::OnPushedAudio(const float* AudioData, int32 NumSamples, int32 NumChannels, int32 SampleRate)
+	{
+		OnRawAudio(AudioData, NumSamples, NumChannels, SampleRate);
+	}
+
+	void FAudioProducer::OnRawAudio(const float* AudioData, int32 NumSamples, int32 NumChannels, int32 SampleRate)
 	{
 		if (bIsMuted)
 		{
 			return;
 		}
 
-		PatchInput->OnNewAudioFrame(AudioData, NumSamples, NumChannels, SampleRate);
-	}
+		TRawAudio<const float> RawAudio = {
+			.Data = AudioData,
+			.NumSamples = NumSamples,
+			.SampleRate = SampleRate,
+			.NumChannels = NumChannels
+		};
 
-	void FAudioProducer::OnPushedAudio(const float* AudioData, int32 NumSamples, int32 NumChannels, int32 SampleRate)
-	{
-		if (bIsMuted)
+		Audio::FAlignedFloatBuffer ProcessedAudioData;
+		FProcessedAudio ProcessedAudio = {
+			.Data = ProcessedAudioData,
+			.NumSamples = 0,
+			.SampleRate = TargetSampleRate,
+			.NumChannels = TargetNumChannels,
+			.Gain = UPixelStreaming2PluginSettings::CVarWebRTCAudioGain.GetValueOnAnyThread()
+		};
+
+		if (!ProcessAudio(RawAudio, ProcessedAudio))
 		{
+			UE_LOGFMT(LogPixelStreaming2, Warning, "FAudioProducer::OnRawAudio: Failed to process audio data.");
 			return;
 		}
 
-		PatchInput->OnNewAudioFrame(AudioData, NumSamples, NumChannels, SampleRate);
+		PushAudio(ProcessedAudio.Data.GetData(), ProcessedAudio.NumSamples);		
 	}
 } // namespace UE::PixelStreaming2
\ No newline at end of file
diff --git a/Engine/Plugins/Media/PixelStreaming2/Source/PixelStreaming2/Private/Blueprints/PixelStreaming2AudioComponent.cpp b/Engine/Plugins/Media/PixelStreaming2/Source/PixelStreaming2/Private/Blueprints/PixelStreaming2AudioComponent.cpp
--- a/Engine/Plugins/Media/PixelStreaming2/Source/PixelStreaming2/Private/Blueprints/PixelStreaming2AudioComponent.cpp
+++ b/Engine/Plugins/Media/PixelStreaming2/Source/PixelStreaming2/Private/Blueprints/PixelStreaming2AudioComponent.cpp
@@ -12,7 +12,7 @@
 #include UE_INLINE_GENERATED_CPP_BY_NAME(PixelStreaming2AudioComponent)
 
 /**
- * Component that recieves audio from a remote webrtc connection and outputs it into UE using a "synth component".
+ * Component that receives audio from a remote webrtc connection and outputs it into UE using a "synth component".
  */
 UPixelStreaming2AudioComponent::UPixelStreaming2AudioComponent(const FObjectInitializer& ObjectInitializer)
 	: Super(ObjectInitializer)
@@ -30,7 +30,7 @@ UPixelStreaming2AudioComponent::UPixelStreaming2AudioComponent(const FObjectInit
 
 ISoundGeneratorPtr UPixelStreaming2AudioComponent::CreateSoundGenerator(const FSoundGeneratorInitParams& InParams)
 {
-	SoundGenerator->SetParameters(InParams);
+	SoundGenerator->Initialize(InParams);
 	Initialize(InParams.SampleRate);
 	return SoundGenerator;
 }
@@ -148,54 +148,15 @@ bool UPixelStreaming2AudioComponent::WillListenToAnyPlayer()
 	return PlayerToHear == FString();
 }
 
-void UPixelStreaming2AudioComponent::ConsumeRawPCM(const int16_t* AudioData, int InSampleRate, size_t NChannels, size_t NFrames)
+void UPixelStreaming2AudioComponent::ConsumeRawPCM(const int16_t* AudioData, int SampleRate, size_t InNumChannels, size_t NumFrames)
 {
 	// Sound generator has not been initialized yet.
-	if (!SoundGenerator || SoundGenerator->GetSampleRate() == 0 || GetAudioComponent() == nullptr)
+	if (!SoundGenerator || !SoundGenerator->IsInitialized())
 	{
 		return;
 	}
 
-	// Set pitch multiplier as a way to handle mismatched sample rates
-	if (InSampleRate != SoundGenerator->GetSampleRate())
-	{
-		GetAudioComponent()->SetPitchMultiplier((float)InSampleRate / SoundGenerator->GetSampleRate());
-	}
-	else if (GetAudioComponent()->PitchMultiplier != 1.0f)
-	{
-		GetAudioComponent()->SetPitchMultiplier(1.0f);
-	}
-
-	//                                   Data       Samples              Channels   SampleRate
-	Audio::TSampleBuffer<int16_t> Buffer(AudioData, NFrames * NChannels, NChannels, InSampleRate);
-	int32						  TargetNumChannels = SoundGenerator->GetNumChannels();
-	if (NChannels != TargetNumChannels)
-	{
-		if (NChannels < TargetNumChannels)
-		{
-			// Up mix by duplicating the mono source to each channel
-			TArrayView<int16> SourceBuffer = Buffer.GetArrayView();
-			TArray<int16>	  MixedBuffer;
-			MixedBuffer.SetNumZeroed(NFrames * TargetNumChannels);
-			for (int32 SrcSampleIdx = 0; SrcSampleIdx < Buffer.GetNumSamples(); SrcSampleIdx++)
-			{
-				int32 DestSampleIdx = SrcSampleIdx * TargetNumChannels;
-				for (int32 Channel = 0; Channel < TargetNumChannels; Channel++)
-				{
-					MixedBuffer[DestSampleIdx + Channel] = SourceBuffer[SrcSampleIdx];
-				}
-			}
-
-			Buffer.CopyFrom(MixedBuffer, TargetNumChannels, InSampleRate);
-		}
-		else
-		{
-			// Down mix using inbuilt method
-			Buffer.MixBufferToChannels(TargetNumChannels);
-		}
-	}
-
-	SoundGenerator->AddAudio(Buffer.GetData(), InSampleRate, Buffer.GetNumChannels(), Buffer.GetNumFrames());
+	SoundGenerator->AddAudio(AudioData, SampleRate, InNumChannels, NumFrames);
 }
 
 void UPixelStreaming2AudioComponent::OnAudioConsumerAdded()
diff --git a/Engine/Plugins/Media/PixelStreaming2/Source/PixelStreaming2/Private/Blueprints/PixelStreaming2AudioComponent.h b/Engine/Plugins/Media/PixelStreaming2/Source/PixelStreaming2/Private/Blueprints/PixelStreaming2AudioComponent.h
--- a/Engine/Plugins/Media/PixelStreaming2/Source/PixelStreaming2/Private/Blueprints/PixelStreaming2AudioComponent.h
+++ b/Engine/Plugins/Media/PixelStreaming2/Source/PixelStreaming2/Private/Blueprints/PixelStreaming2AudioComponent.h
@@ -85,7 +85,7 @@ public:
 	void Reset();
 
 	//~ Begin IPixelStreaming2AudioConsumer interface
-	void ConsumeRawPCM(const int16_t* AudioData, int InSampleRate, size_t NChannels, size_t NFrames);
+	void ConsumeRawPCM(const int16_t* AudioData, int SampleRate, size_t NumChannels, size_t NumFrames);
 	void OnAudioConsumerAdded();
 	void OnAudioConsumerRemoved();
 	//~ End IPixelStreaming2AudioConsumer interface
diff --git a/Engine/Plugins/Media/PixelStreaming2/Source/PixelStreaming2/Private/SoundGenerator.cpp b/Engine/Plugins/Media/PixelStreaming2/Source/PixelStreaming2/Private/SoundGenerator.cpp
--- a/Engine/Plugins/Media/PixelStreaming2/Source/PixelStreaming2/Private/SoundGenerator.cpp
+++ b/Engine/Plugins/Media/PixelStreaming2/Source/PixelStreaming2/Private/SoundGenerator.cpp
@@ -1,11 +1,12 @@
 // Copyright Epic Games, Inc. All Rights Reserved.
 
 #include "SoundGenerator.h"
-#include "DSP/FloatArrayMath.h"
+
+#include "Logging.h"
+#include "UtilsAudio.h"
 
 namespace UE::PixelStreaming2
 {
-
 	FSoundGenerator::FSoundGenerator()
 		: Params()
 		, Buffer()
@@ -13,9 +14,10 @@ namespace UE::PixelStreaming2
 	{
 	}
 
-	void FSoundGenerator::SetParameters(const FSoundGeneratorInitParams& InitParams)
+	void FSoundGenerator::Initialize(const FSoundGeneratorInitParams& InitParams)
 	{
 		Params = InitParams;
+		bInitialized = true;
 	}
 
 	int32 FSoundGenerator::GetDesiredNumSamplesToRenderPerCallback() const
@@ -23,58 +25,88 @@ namespace UE::PixelStreaming2
 		return Params.NumFramesPerCallback * Params.NumChannels;
 	}
 
+	void FSoundGenerator::OnBeginGenerate()
+	{
+		bGeneratingAudio = true;
+	}
+
+	void FSoundGenerator::OnEndGenerate()
+	{
+		bGeneratingAudio = false;
+	}
+
+	bool FSoundGenerator::IsFinished() const
+	{
+		return false;
+	}
+
+	bool FSoundGenerator::IsInitialized()
+	{
+		return bInitialized;
+	}
+
 	void FSoundGenerator::EmptyBuffers()
 	{
 		FScopeLock Lock(&CriticalSection);
 		Buffer.Empty();
 	}
 
-	void FSoundGenerator::AddAudio(const int16_t* AudioData, int InSampleRate, size_t NChannels, size_t NFrames)
+	void FSoundGenerator::AddAudio(const int16* AudioData, int SampleRate, size_t NumChannels, size_t NumFrames)
 	{
 		if (!bGeneratingAudio)
 		{
 			return;
 		}
 
-		int NSamples = NFrames * NChannels;
-
-		// Critical Section as we are writing into the `Buffer` that `ISoundGenerator` is using on another thread.
-		FScopeLock Lock(&CriticalSection);
+		TRawAudio<const int16> RawAudio = {
+			.Data = AudioData,
+			.NumSamples = static_cast<int32>(NumFrames * NumChannels),
+			.SampleRate = SampleRate,
+			.NumChannels = static_cast<int32>(NumChannels),
+		};
+
+		Audio::FAlignedFloatBuffer ProcessedAudioData;
+		FProcessedAudio ProcessedAudio = {
+			.Data = ProcessedAudioData,
+			.NumSamples = 0,
+			.SampleRate = static_cast<int32>(Params.SampleRate),
+			.NumChannels = static_cast<int32>(Params.NumChannels),
+		};
+
+		if (!ProcessAudio(RawAudio, ProcessedAudio))
+		{
+			UE_LOGFMT(LogPixelStreaming2, Warning, "FSoundGenerator::AddAudio: Failed to process audio data.");
+			return;
+		}
 
-		Buffer.Append(AudioData, NSamples);
+		{
+			// Critical Section as we are writing into the `Buffer` that `ISoundGenerator` is using on another thread.
+			FScopeLock Lock(&CriticalSection);
+			Buffer.Append(ProcessedAudioData.GetData(), ProcessedAudioData.Num());
+		}
 	}
 
 	// Called when a new buffer is required.
-	int32 FSoundGenerator::OnGenerateAudio(float* OutAudio, int32 NumSamples)
+	int32 FSoundGenerator::OnGenerateAudio(float* OutAudio, int32 NumRequestedSamples)
 	{
+		FScopeLock Lock(&CriticalSection);
+		
+		FMemory::Memzero(OutAudio, NumRequestedSamples * sizeof(float));
+
 		// Not listening to peer, return zero'd buffer.
 		if (!bShouldGenerateAudio || Buffer.Num() == 0)
 		{
-			FMemory::Memzero(OutAudio, NumSamples * sizeof(float));
-			return NumSamples;
+			return NumRequestedSamples;
 		}
 
-		int32 NumSamplesToCopy = NumSamples;
-
-		// Critical section
-		{
-			FScopeLock Lock(&CriticalSection);
-
-			NumSamplesToCopy = FGenericPlatformMath::Min(NumSamples, Buffer.Num());
-
-			// Copy from local buffer into OutAudio if we have enough samples
-			Audio::ArrayPcm16ToFloat(MakeArrayView(Buffer.GetData(), NumSamplesToCopy), MakeArrayView(OutAudio, NumSamplesToCopy));
+		int32 NumSamplesToCopy = FGenericPlatformMath::Min(NumRequestedSamples, Buffer.Num());
 
-			// Remove front NumSamples from the local buffer
-			Buffer.RemoveAt(0, NumSamplesToCopy, EAllowShrinking::No);
-		}
+		// Copy from local buffer into OutAudio if we have enough samples
+		FMemory::Memcpy(OutAudio, Buffer.GetData(), NumSamplesToCopy * sizeof(float));
 
-		if (NumSamplesToCopy < NumSamples)
-		{
-			FMemory::Memzero(&OutAudio[NumSamplesToCopy], (NumSamples - NumSamplesToCopy) * sizeof(float));
-		}
+		// Remove front NumSamplesToCopy from the local buffer
+		Buffer.RemoveAt(0, NumSamplesToCopy, EAllowShrinking::No);
 
 		return NumSamplesToCopy;
 	}
-
 } // namespace UE::PixelStreaming2
diff --git a/Engine/Plugins/Media/PixelStreaming2/Source/PixelStreaming2Core/Public/IPixelStreaming2AudioConsumer.h b/Engine/Plugins/Media/PixelStreaming2/Source/PixelStreaming2Core/Public/IPixelStreaming2AudioConsumer.h
--- a/Engine/Plugins/Media/PixelStreaming2/Source/PixelStreaming2Core/Public/IPixelStreaming2AudioConsumer.h
+++ b/Engine/Plugins/Media/PixelStreaming2/Source/PixelStreaming2Core/Public/IPixelStreaming2AudioConsumer.h
@@ -24,11 +24,11 @@ public:
 	/**
 	 * @brief Consume raw audio data.
 	 * @param AudioData Pointer to the audio data.
-	 * @param InSampleRate Audio sample rate in samples per second.
-	 * @param NChannels Number of audio channels. For example 2 for stero audio.
-	 * @param NFrames Number of Audio frames in a single channel.
+	 * @param SampleRate Audio sample rate in samples per second.
+	 * @param NumChannels Number of audio channels. For example 2 for stero audio.
+	 * @param NumFrames Number of Audio frames in a single channel.
 	 */
-	virtual void ConsumeRawPCM(const int16_t* AudioData, int InSampleRate, size_t NChannels, size_t NFrames) = 0;
+	virtual void ConsumeRawPCM(const int16_t* AudioData, int SampleRate, size_t NumChannels, size_t NumFrames) = 0;
 
 	/**
 	 * @brief Called when a audio consumer is added.
diff --git a/Engine/Plugins/Media/PixelStreaming2/Source/PixelStreaming2RTC/Private/Blueprints/PixelStreaming2Peer.cpp b/Engine/Plugins/Media/PixelStreaming2/Source/PixelStreaming2RTC/Private/Blueprints/PixelStreaming2Peer.cpp
--- a/Engine/Plugins/Media/PixelStreaming2/Source/PixelStreaming2RTC/Private/Blueprints/PixelStreaming2Peer.cpp
+++ b/Engine/Plugins/Media/PixelStreaming2/Source/PixelStreaming2RTC/Private/Blueprints/PixelStreaming2Peer.cpp
@@ -200,7 +200,7 @@ bool UPixelStreaming2Peer::Subscribe(const FString& StreamerId)
 
 ISoundGeneratorPtr UPixelStreaming2Peer::CreateSoundGenerator(const FSoundGeneratorInitParams& InParams)
 {
-	SoundGenerator->SetParameters(InParams);
+	SoundGenerator->Initialize(InParams);
 	Initialize(InParams.SampleRate);
 
 	return SoundGenerator;
@@ -216,31 +216,15 @@ void UPixelStreaming2Peer::OnEndGenerate()
 	SoundGenerator->bGeneratingAudio = false;
 }
 
-void UPixelStreaming2Peer::ConsumeRawPCM(const int16_t* AudioData, int InSampleRate, size_t NChannels, size_t NFrames)
+void UPixelStreaming2Peer::ConsumeRawPCM(const int16_t* AudioData, int SampleRate, size_t InNumChannels, size_t NumFrames)
 {
 	// Sound generator has not been initialized yet.
-	if (!SoundGenerator || SoundGenerator->GetSampleRate() == 0 || GetAudioComponent() == nullptr)
+	if (!SoundGenerator || !SoundGenerator->IsInitialized())
 	{
 		return;
 	}
 
-	// Set pitch multiplier as a way to handle mismatched sample rates
-	if (InSampleRate != SoundGenerator->GetSampleRate())
-	{
-		GetAudioComponent()->SetPitchMultiplier((float)InSampleRate / SoundGenerator->GetSampleRate());
-	}
-	else if (GetAudioComponent()->PitchMultiplier != 1.0f)
-	{
-		GetAudioComponent()->SetPitchMultiplier(1.0f);
-	}
-
-	Audio::TSampleBuffer<int16_t> Buffer(AudioData, NFrames, NChannels, InSampleRate);
-	if (NChannels != SoundGenerator->GetNumChannels())
-	{
-		Buffer.MixBufferToChannels(SoundGenerator->GetNumChannels());
-	}
-
-	SoundGenerator->AddAudio(Buffer.GetData(), InSampleRate, NChannels, Buffer.GetNumSamples());
+	SoundGenerator->AddAudio(AudioData, SampleRate, InNumChannels, NumFrames);
 }
 
 void UPixelStreaming2Peer::OnAudioConsumerAdded()
diff --git a/Engine/Plugins/Media/PixelStreaming2/Source/PixelStreaming2RTC/Private/Blueprints/PixelStreaming2Peer.h b/Engine/Plugins/Media/PixelStreaming2/Source/PixelStreaming2RTC/Private/Blueprints/PixelStreaming2Peer.h
--- a/Engine/Plugins/Media/PixelStreaming2/Source/PixelStreaming2RTC/Private/Blueprints/PixelStreaming2Peer.h
+++ b/Engine/Plugins/Media/PixelStreaming2/Source/PixelStreaming2RTC/Private/Blueprints/PixelStreaming2Peer.h
@@ -117,7 +117,7 @@ private:
 
 public:
 	// Begin IPixelStreaming2AudioConsumer Callbacks
-	virtual void ConsumeRawPCM(const int16_t* AudioData, int InSampleRate, size_t NChannels, size_t NFrames) override;
+	virtual void ConsumeRawPCM(const int16_t* AudioData, int SampleRate, size_t NumChannels, size_t NumFrames) override;
 	virtual void OnAudioConsumerAdded() override;
 	virtual void OnAudioConsumerRemoved() override;
 	// End IPixelStreaming2AudioConsumer Callbacks
diff --git a/Engine/Plugins/Media/PixelStreaming2/Source/PixelStreaming2RTC/Private/PixelStreaming2RTCPlayer.cpp b/Engine/Plugins/Media/PixelStreaming2/Source/PixelStreaming2RTC/Private/PixelStreaming2RTCPlayer.cpp
--- a/Engine/Plugins/Media/PixelStreaming2/Source/PixelStreaming2RTC/Private/PixelStreaming2RTCPlayer.cpp
+++ b/Engine/Plugins/Media/PixelStreaming2/Source/PixelStreaming2RTC/Private/PixelStreaming2RTCPlayer.cpp
@@ -676,20 +676,25 @@ namespace UE::PixelStreaming2
 	{
 	}
 
-	void FPixelStreaming2RTCStreamPlayer::ConsumeRawPCM(const int16_t* AudioData, int InSampleRate, size_t NChannels, size_t NFrames)
+	void FPixelStreaming2RTCStreamPlayer::ConsumeRawPCM(const int16_t* AudioData, int SampleRate, size_t NumChannels, size_t NumFrames)
 	{
+		if (AudioData == nullptr || SampleRate <= 0 || NumChannels == 0 || NumFrames == 0)
+		{
+			return;
+		}
+
 		TSharedRef<FPixelStreaming2AudioSample, ESPMode::ThreadSafe> AudioSample = AudioSamplePool->AcquireShared();
 
 		TArray<uint8> AudioBuffer;
 
-		size_t NumSamples = NFrames * NChannels;
+		size_t NumSamples = NumFrames * NumChannels;
 		size_t DataSize = NumSamples * sizeof(int16);
 		AudioBuffer.SetNumUninitialized(DataSize);
 
 		FMemory::Memcpy(AudioBuffer.GetData(), AudioData, DataSize);
 
-		FTimespan Duration = (NumSamples * ETimespan::TicksPerSecond) / InSampleRate;
-		AudioSample->Initialize(AudioBuffer.GetData(), DataSize, NChannels, InSampleRate, FTimespan::Zero(), Duration);
+		FTimespan Duration = (NumSamples * ETimespan::TicksPerSecond) / SampleRate;
+		AudioSample->Initialize(AudioBuffer.GetData(), DataSize, NumChannels, SampleRate, FTimespan::Zero(), Duration);
 
 		Samples->AddAudio(AudioSample);
 	}
diff --git a/Engine/Plugins/Media/PixelStreaming2/Source/PixelStreaming2RTC/Private/PixelStreaming2RTCPlayer.h b/Engine/Plugins/Media/PixelStreaming2/Source/PixelStreaming2RTC/Private/PixelStreaming2RTCPlayer.h
--- a/Engine/Plugins/Media/PixelStreaming2/Source/PixelStreaming2RTC/Private/PixelStreaming2RTCPlayer.h
+++ b/Engine/Plugins/Media/PixelStreaming2/Source/PixelStreaming2RTC/Private/PixelStreaming2RTCPlayer.h
@@ -173,7 +173,7 @@ namespace UE::PixelStreaming2
 		// End IPixelStreaming2VideoConsumer
 
 		// Begin IPixelStreaming2AudioConsumer
-		virtual void ConsumeRawPCM(const int16_t* AudioData, int InSampleRate, size_t NChannels, size_t NFrames) override;
+		virtual void ConsumeRawPCM(const int16_t* AudioData, int SampleRate, size_t NumChannels, size_t NumFrames) override;
 		virtual void OnAudioConsumerAdded() override;
 		virtual void OnAudioConsumerRemoved() override;
 		// End IPixelStreaming2AudioConsumer
