[PixelStreaming2] Fix: Improper texture dimensions passing through check

#Synced-CL 48443255

diff --git a/Engine/Plugins/Media/PixelStreaming2/Source/PixelStreaming2/Private/VideoProducerViewportBase.cpp b/Engine/Plugins/Media/PixelStreaming2/Source/PixelStreaming2/Private/VideoProducerViewportBase.cpp
--- a/Engine/Plugins/Media/PixelStreaming2/Source/PixelStreaming2/Private/VideoProducerViewportBase.cpp
+++ b/Engine/Plugins/Media/PixelStreaming2/Source/PixelStreaming2/Private/VideoProducerViewportBase.cpp
@@ -85,7 +85,7 @@ namespace UE::PixelStreaming2
 			FMath::Min(CaptureRect.Height(), FMath::Max((BackBuffer->GetDesc().Extent.Y - VirtualPos.Y), 0))
 		};
 	
-		if (VirtualSize == FIntPoint::ZeroValue)
+		if (VirtualSize.X == 0 || VirtualSize.Y == 0)
 		{
 			return; // discard for copying nothing
 		}
