Delay the Mobile Session Hosting until after the matchmaking phase.
Fixed the issue of CreateOffer failure for H.264 streaming in the preflight Test/Shipping build.

diff --git a/Engine/Plugins/Media/PixelStreaming2/Source/PixelStreaming2RTC/Private/UtilsCodecs.h b/Engine/Plugins/Media/PixelStreaming2/Source/PixelStreaming2RTC/Private/UtilsCodecs.h
--- a/Engine/Plugins/Media/PixelStreaming2/Source/PixelStreaming2RTC/Private/UtilsCodecs.h
+++ b/Engine/Plugins/Media/PixelStreaming2/Source/PixelStreaming2RTC/Private/UtilsCodecs.h
@@ -181,7 +181,7 @@ namespace UE::PixelStreaming2
 					return TOptional<FString>();
 			}
 
-			ProfileString = FString::Printf(TEXT("%s%02x"), *ProfileIdcIopString, Level);
+			ProfileString = FString::Printf(TEXT("%s%02x"), *ProfileIdcIopString, static_cast<uint32>(Level));
 		}
 
 		return ProfileString;
