Fixed a crash when exiting the game before the phase of Engine PostInit by adding null pointer protections for the return values at all sites of calling UPixelStreaming2PluginSettings::Delegates();

diff --git a/Engine/Plugins/Media/PixelStreaming2/Source/PixelStreaming2Settings/Private/PixelStreaming2PluginSettings.cpp b/Engine/Plugins/Media/PixelStreaming2/Source/PixelStreaming2Settings/Private/PixelStreaming2PluginSettings.cpp
--- a/Engine/Plugins/Media/PixelStreaming2/Source/PixelStreaming2Settings/Private/PixelStreaming2PluginSettings.cpp
+++ b/Engine/Plugins/Media/PixelStreaming2/Source/PixelStreaming2Settings/Private/PixelStreaming2PluginSettings.cpp
@@ -253,14 +253,26 @@ TAutoConsoleVariable<bool> UPixelStreaming2PluginSettings::CVarLogStats(
 	TEXT("PixelStreaming2.LogStats"),
 	false,
 	TEXT("Whether to show PixelStreaming stats in the log (default: false)."),
-	FConsoleVariableDelegate::CreateLambda([](IConsoleVariable* Var) { Delegates()->OnLogStatsChanged.Broadcast(Var); }),
+	FConsoleVariableDelegate::CreateLambda([](IConsoleVariable* Var) 
+		{
+			if (FDelegates *Delegations = Delegates())
+			{
+				Delegations->OnLogStatsChanged.Broadcast(Var);
+			}
+		}),
 	ECVF_Default);
 
 TAutoConsoleVariable<FString> UPixelStreaming2PluginSettings::CVarEpicRtcLogFilter(
 	TEXT("PixelStreaming2.EpicRtcLogFilter"),
 	"",
 	TEXT("Double forward slash (\"//\") separated list of regex patterns to filter from the EpicRtc logs (default: \"\")."),
-	FConsoleVariableDelegate::CreateLambda([](IConsoleVariable* Var) { Delegates()->OnEpicRtcLogFilterChanged.Broadcast(Var); }),
+	FConsoleVariableDelegate::CreateLambda([](IConsoleVariable* Var) 
+		{
+			if (FDelegates* Delegations = Delegates())
+			{
+				Delegations->OnEpicRtcLogFilterChanged.Broadcast(Var); 
+			}
+		}),
 	ECVF_Default);
 
 TAutoConsoleVariable<bool> UPixelStreaming2PluginSettings::CVarDisableLatencyTester(
@@ -279,7 +291,13 @@ TAutoConsoleVariable<bool> UPixelStreaming2PluginSettings::CVarDecoupleFramerate
 	TEXT("PixelStreaming2.DecoupleFramerate"),
 	false,
 	TEXT("Whether we should only stream as fast as we render or at some fixed interval. Coupled means only stream what we render."),
-	FConsoleVariableDelegate::CreateLambda([](IConsoleVariable* Var) { Delegates()->OnDecoupleFramerateChanged.Broadcast(Var); }),
+	FConsoleVariableDelegate::CreateLambda([](IConsoleVariable* Var) 
+		{
+			if (FDelegates* Delegations = Delegates())
+			{
+				Delegations->OnDecoupleFramerateChanged.Broadcast(Var); 
+			}
+		}),
 	ECVF_Default);
 
 TAutoConsoleVariable<float> UPixelStreaming2PluginSettings::CVarDecoupleWaitFactor(
@@ -310,7 +328,13 @@ TAutoConsoleVariable<bool> UPixelStreaming2PluginSettings::CVarUseMediaCapture(
 	TEXT("PixelStreaming2.UseMediaCapture"),
 	true,
 	TEXT("Use Media Capture from MediaIOFramework to capture frames rather than Pixel Streamings internal backbuffer sources."),
-	FConsoleVariableDelegate::CreateLambda([](IConsoleVariable* Var) { Delegates()->OnUseMediaCaptureChanged.Broadcast(Var); }),
+	FConsoleVariableDelegate::CreateLambda([](IConsoleVariable* Var) 
+		{
+			if (FDelegates* Delegations = Delegates())
+			{
+				Delegations->OnUseMediaCaptureChanged.Broadcast(Var); 
+			}
+		}),
 	ECVF_Default);
 
 TAutoConsoleVariable<FString> UPixelStreaming2PluginSettings::CVarDefaultStreamerID(
@@ -344,14 +368,26 @@ TAutoConsoleVariable<bool> UPixelStreaming2PluginSettings::CVarCaptureUseFence(
 	TEXT("PixelStreaming2.CaptureUseFence"),
 	true,
 	TEXT("Whether the texture copy we do during image capture should use a fence or not (non-fenced is faster but less safe)."),
-	FConsoleVariableDelegate::CreateLambda([](IConsoleVariable* Var) { Delegates()->OnCaptureUseFenceChanged.Broadcast(Var); }),
+	FConsoleVariableDelegate::CreateLambda([](IConsoleVariable* Var) 
+		{
+			if (FDelegates* Delegations = Delegates())
+			{
+				Delegations->OnCaptureUseFenceChanged.Broadcast(Var); 
+			}
+		}),
 	ECVF_Default);
 
 TAutoConsoleVariable<bool> UPixelStreaming2PluginSettings::CVarDebugDumpAudio(
 	TEXT("PixelStreaming2.DumpDebugAudio"),
 	false,
 	TEXT("Dumps mixed audio from PS2 to a file on disk for debugging purposes."),
-	FConsoleVariableDelegate::CreateLambda([](IConsoleVariable* Var) { Delegates()->OnDebugDumpAudioChanged.Broadcast(Var); }),
+	FConsoleVariableDelegate::CreateLambda([](IConsoleVariable* Var) 
+		{
+			if (FDelegates* Delegations = Delegates())
+			{
+				Delegations->OnDebugDumpAudioChanged.Broadcast(Var); 
+			}
+		}),
 	ECVF_Default);
 
 // Begin Encoder CVars
@@ -403,7 +439,14 @@ TAutoConsoleVariable<bool> UPixelStreaming2PluginSettings::CVarEncoderEnableSimu
 	TEXT("PixelStreaming2.Encoder.EnableSimulcast"),
 	false,
 	TEXT("Enables simulcast. When enabled, the encoder will encode at full resolution, 1/2 resolution and 1/4 resolution simultaneously. Note: Simulcast is only supported with `H264` and `VP8` and you must use the SFU from the infrastructure to fully utilise this functionality."),
-	FConsoleVariableDelegate::CreateLambda([](IConsoleVariable* Var) { VerifyCVarVideoSettings(nullptr); Delegates()->OnSimulcastEnabledChanged.Broadcast(Var); }),
+	FConsoleVariableDelegate::CreateLambda([](IConsoleVariable* Var) 
+		{
+			VerifyCVarVideoSettings(nullptr); 
+			if (FDelegates* Delegations = Delegates())
+			{
+				Delegations->OnSimulcastEnabledChanged.Broadcast(Var);
+			}
+		}),
 	ECVF_Default);
 
 FString UE::PixelStreaming2::GSelectedCodec = TEXT("H264");
@@ -419,7 +462,14 @@ TAutoConsoleVariable<FString> UPixelStreaming2PluginSettings::CVarEncoderScalabi
 	TEXT("PixelStreaming2.Encoder.ScalabilityMode"),
 	TEXT("L1T1"),
 	TEXT("Indicates number of Spatial and temporal layers used, default: L1T1. For a full list of values refer to https://www.w3.org/TR/webrtc-svc/#scalabilitymodes*"),
-	FConsoleVariableDelegate::CreateLambda([](IConsoleVariable* Var) { VerifyCVarVideoSettings(nullptr); Delegates()->OnScalabilityModeChanged.Broadcast(Var); }),
+	FConsoleVariableDelegate::CreateLambda([](IConsoleVariable* Var) 
+		{
+			VerifyCVarVideoSettings(nullptr); 
+			if (FDelegates* Delegations = Delegates())
+			{
+				Delegations->OnScalabilityModeChanged.Broadcast(Var);
+			}
+		}),
 	ECVF_Default);
 
 TAutoConsoleVariable<FString> UPixelStreaming2PluginSettings::CVarEncoderH264Profile(
@@ -432,7 +482,13 @@ TAutoConsoleVariable<bool> UPixelStreaming2PluginSettings::CVarEncoderDebugDumpF
 	TEXT("PixelStreaming2.Encoder.DumpDebugFrames"),
 	false,
 	TEXT("Dumps frames from the encoder to a file on disk for debugging purposes."),
-	FConsoleVariableDelegate::CreateLambda([](IConsoleVariable* Var) { Delegates()->OnEncoderDebugDumpFrameChanged.Broadcast(Var); }),
+	FConsoleVariableDelegate::CreateLambda([](IConsoleVariable* Var)
+		{
+			if (FDelegates* Delegations = Delegates())
+			{
+				Delegations->OnEncoderDebugDumpFrameChanged.Broadcast(Var); 
+			}
+		}),
 	ECVF_Default);
 
 // Begin WebRTC CVars
@@ -441,7 +497,13 @@ TAutoConsoleVariable<int32> UPixelStreaming2PluginSettings::CVarWebRTCFps(
 	TEXT("PixelStreaming2.WebRTC.Fps"),
 	60,
 	TEXT("Framerate for WebRTC encoding. Default: 60"),
-	FConsoleVariableDelegate::CreateLambda([](IConsoleVariable* Var) { Delegates()->OnWebRTCFpsChanged.Broadcast(Var); }),
+	FConsoleVariableDelegate::CreateLambda([](IConsoleVariable* Var) 
+		{
+			if (FDelegates* Delegations = Delegates())
+			{
+				Delegations->OnWebRTCFpsChanged.Broadcast(Var); 
+			}
+		}),
 	ECVF_Default);
 
 // Note: 1 megabit is the maximum allowed in WebRTC for a start bitrate.
@@ -455,14 +517,26 @@ TAutoConsoleVariable<int32> UPixelStreaming2PluginSettings::CVarWebRTCMinBitrate
 	TEXT("PixelStreaming2.WebRTC.MinBitrate"),
 	100000,
 	TEXT("Min bitrate (bps) that WebRTC will not request below. Careful not to set too high otherwise WebRTC will just drop frames. Default: 100000"),
-	FConsoleVariableDelegate::CreateLambda([](IConsoleVariable* Var) { Delegates()->OnWebRTCBitrateChanged.Broadcast(Var); }),
+	FConsoleVariableDelegate::CreateLambda([](IConsoleVariable* Var) 
+		{ 
+			if (FDelegates* Delegations = Delegates())
+			{ 
+				Delegations->OnWebRTCBitrateChanged.Broadcast(Var); 
+			} 
+		}),
 	ECVF_RenderThreadSafe);
 
 TAutoConsoleVariable<int32> UPixelStreaming2PluginSettings::CVarWebRTCMaxBitrate(
 	TEXT("PixelStreaming2.WebRTC.MaxBitrate"),
 	40000000,
 	TEXT("Max bitrate (bps) that WebRTC will not request above. Default: 40000000 aka 40 megabits/per second."),
-	FConsoleVariableDelegate::CreateLambda([](IConsoleVariable* Var) { Delegates()->OnWebRTCBitrateChanged.Broadcast(Var); }),
+	FConsoleVariableDelegate::CreateLambda([](IConsoleVariable* Var) 
+		{ 
+			if (FDelegates* Delegations = Delegates())
+			{ 
+				Delegations->OnWebRTCBitrateChanged.Broadcast(Var); 
+			} 
+		}),
 	ECVF_RenderThreadSafe);
 
 TAutoConsoleVariable<bool> UPixelStreaming2PluginSettings::CVarWebRTCDisableReceiveAudio(
@@ -505,7 +579,13 @@ TAutoConsoleVariable<bool> UPixelStreaming2PluginSettings::CVarWebRTCDisableStat
 	TEXT("PixelStreaming2.WebRTC.DisableStats"),
 	false,
 	TEXT("Disables the collection of WebRTC stats."),
-	FConsoleVariableDelegate::CreateLambda([](IConsoleVariable* Var) { Delegates()->OnWebRTCDisableStatsChanged.Broadcast(Var); }),
+	FConsoleVariableDelegate::CreateLambda([](IConsoleVariable* Var)
+		{
+			if (FDelegates* Delegations = Delegates())
+			{ 
+				Delegations->OnWebRTCDisableStatsChanged.Broadcast(Var); 
+			} 
+		}),
 	ECVF_Default);
 
 TAutoConsoleVariable<float> UPixelStreaming2PluginSettings::CVarWebRTCStatsInterval(
@@ -628,7 +708,13 @@ TAutoConsoleVariable<FString> UPixelStreaming2PluginSettings::CVarInputKeyFilter
 	TEXT("PixelStreaming2.KeyFilter"),
 	"",
 	TEXT("Comma separated list of keys to ignore from streaming clients. Default: \"\""),
-	FConsoleVariableDelegate::CreateLambda([](IConsoleVariable* Var) { Delegates()->OnInputKeyFilterChanged.Broadcast(Var); }),
+	FConsoleVariableDelegate::CreateLambda([](IConsoleVariable* Var) 
+		{ 
+			if (FDelegates* Delegations = Delegates())
+			{
+				Delegations->OnInputKeyFilterChanged.Broadcast(Var); 
+			}
+		}),
 	ECVF_Default);
 // End Input CVars
 
