[PixelStreaming2] Fix: Potential deadlocks during player connection

#Synced-CL 48443255

diff --git a/Engine/Plugins/Media/PixelStreaming2/Source/PixelStreaming2RTC/Private/EpicRtcStreamer.cpp b/Engine/Plugins/Media/PixelStreaming2/Source/PixelStreaming2RTC/Private/EpicRtcStreamer.cpp
--- a/Engine/Plugins/Media/PixelStreaming2/Source/PixelStreaming2RTC/Private/EpicRtcStreamer.cpp
+++ b/Engine/Plugins/Media/PixelStreaming2/Source/PixelStreaming2RTC/Private/EpicRtcStreamer.cpp
@@ -640,17 +640,17 @@ namespace UE::PixelStreaming2
 			return;
 		}
 
-		Participants->Apply([this, ConnectionInterface](FString PlayerId, TSharedPtr<FPlayerContext> Participant) {
-			if (!ConnectionInterface)
-			{
-				return;
-			}
-
+		TSet<FString> PlayerIds;
+		Participants->GetKeys(PlayerIds);
+		// We need to iterate over just the keys as we can potentially deadlock during player connection with OnDataTrackState having locked the peer connection
+		// and SetConnectionRates needing to use the peer connection. Iterating over just the keys is the easiest way to prevent this
+		for (const FString& PlayerId : PlayerIds)
+		{
 			// We don't include the starting bitrate during a bitrate refresh as that causes the peer connection
 			// to go back to that value and have to work its way up. This function is also only triggered by changes to min and max
 			FUtf8String Utf8PlayerId = *PlayerId;
 			ConnectionInterface->SetConnectionRates(ToEpicRtcStringView(Utf8PlayerId), GetBitrates(false /* bIncludeStartBitrate */));
-		});
+		}
 	}
 
 	void FEpicRtcStreamer::ForEachPlayer(const TFunction<void(FString, TSharedPtr<FPlayerContext>)>& Func)
