[PixelStreaming2] QOL: Public API changes
- Added default constructor and destructor for IPixelStreaming2InputHandler that registers with the FInputDevice. This means any user implementation will be ticked
- Moved DefaultDataProtocol, InputStructures and JavaScriptKeyCodes.inl to Public
- Added SetInputHandler to the IPixelStreaming2Streamer interface (and updated usage in DummyStreamer and EpicRtcStreamer)
- Added Create methods for existing PixelStreaming2 video producer implementations

#Synced-CL 48443255

#Synced-CL 48443255

diff --git a/Engine/Plugins/Media/PixelStreaming2/Source/PixelStreaming2/Private/PixelStreaming2VideoProducers.cpp b/Engine/Plugins/Media/PixelStreaming2/Source/PixelStreaming2/Private/PixelStreaming2VideoProducers.cpp
new file mode 100644
--- /dev/null
+++ b/Engine/Plugins/Media/PixelStreaming2/Source/PixelStreaming2/Private/PixelStreaming2VideoProducers.cpp
@@ -0,0 +1,25 @@
+// Copyright Epic Games, Inc. All Rights Reserved.
+
+#include "PixelStreaming2VideoProducers.h"
+
+#include "VideoProducerBackBuffer.h"
+#include "VideoProducerRenderTarget.h"
+#include "VideoProducerPIEViewport.h"
+
+namespace UE::PixelStreaming2
+{
+    TSharedPtr<IPixelStreaming2VideoProducer> CreateVideoProducerBackBuffer()
+    {
+        return FVideoProducerBackBuffer::Create();
+    }
+
+    TSharedPtr<IPixelStreaming2VideoProducer> CreateVideoProducerRenderTarget(UTextureRenderTarget2D* Target)
+    {
+        return FVideoProducerRenderTarget::Create(Target);
+    }
+
+    TSharedPtr<IPixelStreaming2VideoProducer> CreateVideoProducerPIEViewport()
+    {
+        return FVideoProducerPIEViewport::Create();
+    }
+}
\ No newline at end of file
diff --git a/Engine/Plugins/Media/PixelStreaming2/Source/PixelStreaming2/Private/VideoProducerRenderTarget.cpp b/Engine/Plugins/Media/PixelStreaming2/Source/PixelStreaming2/Private/VideoProducerRenderTarget.cpp
--- a/Engine/Plugins/Media/PixelStreaming2/Source/PixelStreaming2/Private/VideoProducerRenderTarget.cpp
+++ b/Engine/Plugins/Media/PixelStreaming2/Source/PixelStreaming2/Private/VideoProducerRenderTarget.cpp
@@ -28,18 +28,29 @@ namespace UE::PixelStreaming2
 
 	void FVideoProducerRenderTarget::OnEndFrameRenderThread()
 	{
-		if (Target)
+		if (!Target)
 		{
-			if (FTextureRHIRef Texture = Target->GetResource()->GetTexture2DRHI())
-			{
-				PushFrame(FPixelCaptureInputFrameRHI(Texture));
-			}
+			return;
 		}
+
+		FTextureResource* Resource = Target->GetResource();
+		if (!Resource)
+		{
+			return;
+		}
+
+		FTextureRHIRef Texture = Resource->GetTexture2DRHI();
+		if (!Texture)
+		{
+			return;
+		}
+		
+		PushFrame(FPixelCaptureInputFrameRHI(Texture));
 	}
 
 	FString FVideoProducerRenderTarget::ToString()
 	{
-		return VideoProducerIdentifiers::FVideoProducerRenderTarget;
+		return TEXT("CustomVideoProducer");
 	}
 
 } // namespace UE::PixelStreaming2
\ No newline at end of file
diff --git a/Engine/Plugins/Media/PixelStreaming2/Source/PixelStreaming2/Public/PixelStreaming2VideoProducers.h b/Engine/Plugins/Media/PixelStreaming2/Source/PixelStreaming2/Public/PixelStreaming2VideoProducers.h
new file mode 100644
--- /dev/null
+++ b/Engine/Plugins/Media/PixelStreaming2/Source/PixelStreaming2/Public/PixelStreaming2VideoProducers.h
@@ -0,0 +1,18 @@
+// Copyright Epic Games, Inc. All Rights Reserved.
+
+#pragma once
+
+#include "CoreTypes.h"
+#include "IPixelStreaming2VideoProducer.h"
+#include "TextureResource.h"
+
+#define UE_API PIXELSTREAMING2_API
+
+namespace UE::PixelStreaming2
+{
+    UE_API TSharedPtr<IPixelStreaming2VideoProducer> CreateVideoProducerBackBuffer();
+    UE_API TSharedPtr<IPixelStreaming2VideoProducer> CreateVideoProducerRenderTarget(UTextureRenderTarget2D* Target);
+    UE_API TSharedPtr<IPixelStreaming2VideoProducer> CreateVideoProducerPIEViewport();
+}
+
+#undef UE_API
\ No newline at end of file
diff --git a/Engine/Plugins/Media/PixelStreaming2/Source/PixelStreaming2Core/Private/Streamer.cpp b/Engine/Plugins/Media/PixelStreaming2/Source/PixelStreaming2Core/Private/Streamer.cpp
--- a/Engine/Plugins/Media/PixelStreaming2/Source/PixelStreaming2Core/Private/Streamer.cpp
+++ b/Engine/Plugins/Media/PixelStreaming2/Source/PixelStreaming2Core/Private/Streamer.cpp
@@ -43,6 +43,7 @@ namespace UE::PixelStreaming2
 		virtual void										 KickPlayer(FString PlayerId) override {}
 		virtual TArray<FString>								 GetConnectedPlayers() override { return {}; }
 		virtual TWeakPtr<class IPixelStreaming2InputHandler> GetInputHandler() override { return nullptr; }
+		virtual void SetInputHandler(TSharedPtr<class IPixelStreaming2InputHandler> Handler) override {}
 		virtual TWeakPtr<IPixelStreaming2AudioSink>			 GetPeerAudioSink(FString PlayerId) override { return nullptr; }
 		virtual TWeakPtr<IPixelStreaming2AudioSink>			 GetUnlistenedAudioSink() override { return nullptr; }
 		virtual TWeakPtr<IPixelStreaming2VideoSink>			 GetPeerVideoSink(FString PlayerId) override { return nullptr; }
diff --git a/Engine/Plugins/Media/PixelStreaming2/Source/PixelStreaming2Core/Public/IPixelStreaming2Streamer.h b/Engine/Plugins/Media/PixelStreaming2/Source/PixelStreaming2Core/Public/IPixelStreaming2Streamer.h
--- a/Engine/Plugins/Media/PixelStreaming2/Source/PixelStreaming2Core/Public/IPixelStreaming2Streamer.h
+++ b/Engine/Plugins/Media/PixelStreaming2/Source/PixelStreaming2Core/Public/IPixelStreaming2Streamer.h
@@ -233,6 +233,12 @@ public:
 	 */
 	virtual TWeakPtr<class IPixelStreaming2InputHandler> GetInputHandler() = 0;
 
+	/**
+	 * @brief Set the streamer's input handler
+	 * @param Handler The streamer's input handler
+	 */
+	virtual void SetInputHandler(TSharedPtr<class IPixelStreaming2InputHandler> Handler) = 0;
+
 	/**
 	 * @brief Get the audio sink associated with a specific peer/player.
 	 * @param PlayerId The player id of the audio sink to retrieve.
diff --git a/Engine/Plugins/Media/PixelStreaming2/Source/PixelStreaming2Input/Internal/DefaultDataProtocol.h b/Engine/Plugins/Media/PixelStreaming2/Source/PixelStreaming2Input/Internal/DefaultDataProtocol.h
deleted file mode 100644
diff --git a/Engine/Plugins/Media/PixelStreaming2/Source/PixelStreaming2Input/Internal/InputStructures.h b/Engine/Plugins/Media/PixelStreaming2/Source/PixelStreaming2Input/Internal/InputStructures.h
deleted file mode 100644
diff --git a/Engine/Plugins/Media/PixelStreaming2/Source/PixelStreaming2Input/Internal/JavaScriptKeyCodes.inl b/Engine/Plugins/Media/PixelStreaming2/Source/PixelStreaming2Input/Internal/JavaScriptKeyCodes.inl
deleted file mode 100644
diff --git a/Engine/Plugins/Media/PixelStreaming2/Source/PixelStreaming2Input/Private/DefaultDataProtocol.cpp b/Engine/Plugins/Media/PixelStreaming2/Source/PixelStreaming2Input/Private/DefaultDataProtocol.cpp
deleted file mode 100644
diff --git a/Engine/Plugins/Media/PixelStreaming2/Source/PixelStreaming2Input/Private/InputHandler.cpp b/Engine/Plugins/Media/PixelStreaming2/Source/PixelStreaming2Input/Private/InputHandler.cpp
--- a/Engine/Plugins/Media/PixelStreaming2/Source/PixelStreaming2Input/Private/InputHandler.cpp
+++ b/Engine/Plugins/Media/PixelStreaming2/Source/PixelStreaming2Input/Private/InputHandler.cpp
@@ -3,7 +3,7 @@
 #include "InputHandler.h"
 
 #include "ApplicationWrapper.h"
-#include "InputStructures.h"
+#include "PixelStreaming2InputStructures.h"
 #include "Framework/Application/SlateApplication.h"
 #include "Policies/CondensedJsonPrintPolicy.h"
 #include "Serialization/JsonSerializer.h"
@@ -45,9 +45,6 @@ namespace UE::PixelStreaming2Input
 		{
 			OnInputKeyFilterChangedHandle = Delegates->OnInputKeyFilterChanged.AddRaw(this, &FPixelStreaming2InputHandler::OnFilteredKeysChanged);
 		}
-
-		// Register this input handler with the module's input device so that it's ticked
-		FInputDevice::GetInputDevice()->AddInputHandler(this);
 	}
 
 	FPixelStreaming2InputHandler::~FPixelStreaming2InputHandler()
@@ -58,8 +55,6 @@ namespace UE::PixelStreaming2Input
 		{
 			Delegates->OnInputKeyFilterChanged.Remove(OnInputKeyFilterChangedHandle);
 		}
-
-		FInputDevice::GetInputDevice()->RemoveInputHandler(this);
 	}
 
 	TSharedPtr<IPixelStreaming2DataProtocol> FPixelStreaming2InputHandler::GetToStreamerProtocol()
@@ -175,6 +170,11 @@ namespace UE::PixelStreaming2Input
 
 	void FPixelStreaming2InputHandler::OnMessage(FString SourceId, TArray<uint8> Buffer)
 	{
+		if (Buffer.Num() == 0)
+		{
+			return;
+		}
+		
 		uint8 MessageType = Buffer[0];
 		// Remove the message type. The remaining data in the buffer is now purely
 		// the message data
diff --git a/Engine/Plugins/Media/PixelStreaming2/Source/PixelStreaming2Input/Private/PixelStreaming2DefaultDataProtocol.cpp b/Engine/Plugins/Media/PixelStreaming2/Source/PixelStreaming2Input/Private/PixelStreaming2DefaultDataProtocol.cpp
new file mode 100644
--- /dev/null
+++ b/Engine/Plugins/Media/PixelStreaming2/Source/PixelStreaming2Input/Private/PixelStreaming2DefaultDataProtocol.cpp
@@ -0,0 +1,167 @@
+// Copyright Epic Games, Inc. All Rights Reserved.
+
+#include "PixelStreaming2DefaultDataProtocol.h"
+
+#include "InputProtocolMap.h"
+#include "IPixelStreaming2InputMessage.h"
+
+namespace UE::PixelStreaming2Input
+{
+
+	typedef EPixelStreaming2MessageTypes EType;
+
+	TSharedPtr<IPixelStreaming2DataProtocol> GetDefaultToStreamerProtocol()
+	{
+		TSharedPtr<FInputProtocolMap> ToStreamerProtocol = TSharedPtr<FInputProtocolMap>(new FInputProtocolMap(EPixelStreaming2MessageDirection::ToStreamer));
+
+		/**
+		 * Control Messages.
+		 */
+		// Simple commands with no payload
+		ToStreamerProtocol->AddInternal(EPixelStreaming2ToStreamerMessage::IFrameRequest, 0);
+		ToStreamerProtocol->AddInternal(EPixelStreaming2ToStreamerMessage::RequestQualityControl, 1);
+		ToStreamerProtocol->AddInternal(EPixelStreaming2ToStreamerMessage::FpsRequest, 2);
+		ToStreamerProtocol->AddInternal(EPixelStreaming2ToStreamerMessage::AverageBitrateRequest, 3);
+		ToStreamerProtocol->AddInternal(EPixelStreaming2ToStreamerMessage::StartStreaming, 4);
+		ToStreamerProtocol->AddInternal(EPixelStreaming2ToStreamerMessage::StopStreaming, 5);
+		ToStreamerProtocol->AddInternal(EPixelStreaming2ToStreamerMessage::LatencyTest, 6, { EType::String });
+		ToStreamerProtocol->AddInternal(EPixelStreaming2ToStreamerMessage::RequestInitialSettings, 7);
+		ToStreamerProtocol->AddInternal(EPixelStreaming2ToStreamerMessage::TestEcho, 8, { EType::String });
+
+		/**
+		 * Input Messages.
+		 */
+		// Generic Input Messages.
+		ToStreamerProtocol->AddInternal(EPixelStreaming2ToStreamerMessage::UIInteraction, 50, { EType::String });
+		ToStreamerProtocol->AddInternal(EPixelStreaming2ToStreamerMessage::Command, 51, { EType::String });
+		ToStreamerProtocol->AddInternal(EPixelStreaming2ToStreamerMessage::TextboxEntry, 52, { EType::String });
+
+		// Keyboard Input Message.
+		// Complex command with payload, therefore we specify the length of the payload (bytes) as well as the structure of the payload
+		ToStreamerProtocol->AddInternal(EPixelStreaming2ToStreamerMessage::KeyDown, 60, { EType::Uint8, EType::Uint8 });
+		ToStreamerProtocol->AddInternal(EPixelStreaming2ToStreamerMessage::KeyUp, 61, { EType::Uint8 });
+		ToStreamerProtocol->AddInternal(EPixelStreaming2ToStreamerMessage::KeyPress, 62, { EType::Uint16 });
+
+		// Mouse Input Messages.
+		ToStreamerProtocol->AddInternal(EPixelStreaming2ToStreamerMessage::MouseEnter, 70);
+		ToStreamerProtocol->AddInternal(EPixelStreaming2ToStreamerMessage::MouseLeave, 71);
+		ToStreamerProtocol->AddInternal(EPixelStreaming2ToStreamerMessage::MouseDown, 72, { EType::Uint8, EType::Uint16, EType::Uint16 });
+		ToStreamerProtocol->AddInternal(EPixelStreaming2ToStreamerMessage::MouseUp, 73, { EType::Uint8, EType::Uint16, EType::Uint16 });
+		ToStreamerProtocol->AddInternal(EPixelStreaming2ToStreamerMessage::MouseMove, 74, { EType::Uint16, EType::Uint16, EType::Uint16, EType::Uint16 });
+		ToStreamerProtocol->AddInternal(EPixelStreaming2ToStreamerMessage::MouseWheel, 75, { EType::Int16, EType::Uint16, EType::Uint16 });
+		ToStreamerProtocol->AddInternal(EPixelStreaming2ToStreamerMessage::MouseDouble, 76, { EType::Uint8, EType::Uint16, EType::Uint16 });
+
+		// Touch Input Messages.
+		ToStreamerProtocol->AddInternal(EPixelStreaming2ToStreamerMessage::TouchStart, 80, { EType::Uint8, EType::Uint16, EType::Uint16, EType::Uint8, EType::Uint8, EType::Uint8 });
+		ToStreamerProtocol->AddInternal(EPixelStreaming2ToStreamerMessage::TouchEnd, 81, { EType::Uint8, EType::Uint16, EType::Uint16, EType::Uint8, EType::Uint8, EType::Uint8 });
+		ToStreamerProtocol->AddInternal(EPixelStreaming2ToStreamerMessage::TouchMove, 82, { EType::Uint8, EType::Uint16, EType::Uint16, EType::Uint8, EType::Uint8, EType::Uint8 });
+
+		// Gamepad Input Messages.
+		ToStreamerProtocol->AddInternal(EPixelStreaming2ToStreamerMessage::GamepadButtonPressed, 90, { EType::Uint8, EType::Uint8, EType::Uint8 });
+		ToStreamerProtocol->AddInternal(EPixelStreaming2ToStreamerMessage::GamepadButtonReleased, 91, { EType::Uint8, EType::Uint8, EType::Uint8 });
+		ToStreamerProtocol->AddInternal(EPixelStreaming2ToStreamerMessage::GamepadAnalog, 92, { EType::Uint8, EType::Uint8, EType::Double });
+		ToStreamerProtocol->AddInternal(EPixelStreaming2ToStreamerMessage::GamepadConnected, 93);
+		ToStreamerProtocol->AddInternal(EPixelStreaming2ToStreamerMessage::GamepadDisconnected, 94, { EType::Uint8 });
+
+		// XR Input Messages.
+		// clang-format off
+		ToStreamerProtocol->AddInternal(EPixelStreaming2ToStreamerMessage::XREyeViews, 109, {
+															// Left eye transform (4x4 matrix)
+															EType::Float, EType::Float, EType::Float, EType::Float,
+															EType::Float, EType::Float, EType::Float, EType::Float,
+															EType::Float, EType::Float, EType::Float, EType::Float,
+															EType::Float, EType::Float, EType::Float, EType::Float,
+															// Left eye perspective projection (4x4 matrix)
+															EType::Float, EType::Float, EType::Float, EType::Float,
+															EType::Float, EType::Float, EType::Float, EType::Float,
+															EType::Float, EType::Float, EType::Float, EType::Float,
+															EType::Float, EType::Float, EType::Float, EType::Float,
+															// Right eye transform (4x4 matrix)
+															EType::Float, EType::Float, EType::Float, EType::Float,
+															EType::Float, EType::Float, EType::Float, EType::Float,
+															EType::Float, EType::Float, EType::Float, EType::Float,
+															EType::Float, EType::Float, EType::Float, EType::Float,
+															// Right eye perspective projection (4x4 matrix)
+															EType::Float, EType::Float, EType::Float, EType::Float,
+															EType::Float, EType::Float, EType::Float, EType::Float,
+															EType::Float, EType::Float, EType::Float, EType::Float,
+															EType::Float, EType::Float, EType::Float, EType::Float,
+															// HMD 4x4 Transform
+															EType::Float, EType::Float, EType::Float, EType::Float,
+															EType::Float, EType::Float, EType::Float, EType::Float,
+															EType::Float, EType::Float, EType::Float, EType::Float,
+															EType::Float, EType::Float, EType::Float, EType::Float,
+														});
+		ToStreamerProtocol->AddInternal(EPixelStreaming2ToStreamerMessage::XRHMDTransform, 110, {	// 4x4 Transform
+															EType::Float, EType::Float, EType::Float, EType::Float,
+															EType::Float, EType::Float, EType::Float, EType::Float,
+															EType::Float, EType::Float, EType::Float, EType::Float,
+															EType::Float, EType::Float, EType::Float, EType::Float,
+														});
+
+		ToStreamerProtocol->AddInternal(EPixelStreaming2ToStreamerMessage::XRControllerTransform, 111, {// 4x4 Transform
+															EType::Float, EType::Float, EType::Float, EType::Float,
+															EType::Float, EType::Float, EType::Float, EType::Float,
+															EType::Float, EType::Float, EType::Float, EType::Float,
+															EType::Float, EType::Float, EType::Float, EType::Float,
+															// Handedness (L, R, Any)
+															EType::Uint8
+														});
+
+		ToStreamerProtocol->AddInternal(EPixelStreaming2ToStreamerMessage::XRButtonPressed, 112, {// Handedness,  ButtonIdx,    IsRepeat,     PressedAmount
+															EType::Uint8, EType::Uint8, EType::Uint8, EType::Double
+														});
+
+		ToStreamerProtocol->AddInternal(EPixelStreaming2ToStreamerMessage::XRButtonTouched, 113, {// Handedness,  ButtonIdx,    IsRepeat
+															EType::Uint8, EType::Uint8, EType::Uint8
+														});
+
+		ToStreamerProtocol->AddInternal(EPixelStreaming2ToStreamerMessage::XRButtonReleased, 114, {// Handedness, ButtonIdx,    IsRepeat
+															EType::Uint8, EType::Uint8, EType::Uint8
+														});
+
+		ToStreamerProtocol->AddInternal(EPixelStreaming2ToStreamerMessage::XRAnalog, 115, {// Handedness, ButtonIdx, AxisValue
+															EType::Uint8, EType::Uint8, EType::Double
+														});
+
+		ToStreamerProtocol->AddInternal(EPixelStreaming2ToStreamerMessage::XRSystem, 116, {// Type of the XR system
+															EType::Uint8
+														});
+
+		ToStreamerProtocol->AddInternal(EPixelStreaming2ToStreamerMessage::XRButtonTouchReleased, 117, {// Handedness, ButtonIdx, IsRepeat
+															EType::Uint8, EType::Uint8, EType::Uint8
+														});
+		
+		ToStreamerProtocol->AddInternal(EPixelStreaming2ToStreamerMessage::ChannelRelayStatus, 198); // id, 2 byte string length, string, uint8 flag
+		
+		ToStreamerProtocol->AddInternal(EPixelStreaming2ToStreamerMessage::Multiplexed, 199); // id, 2 byte string length, string, original message
+
+		// clang-format on
+		return ToStreamerProtocol;
+	}
+
+	TSharedPtr<IPixelStreaming2DataProtocol> GetDefaultFromStreamerProtocol()
+	{
+		TSharedPtr<FInputProtocolMap> FromStreamerProtocol = TSharedPtr<FInputProtocolMap>(new FInputProtocolMap(EPixelStreaming2MessageDirection::FromStreamer));
+
+		FromStreamerProtocol->AddInternal(EPixelStreaming2FromStreamerMessage::QualityControlOwnership, 0);
+		FromStreamerProtocol->AddInternal(EPixelStreaming2FromStreamerMessage::Response, 1);
+		FromStreamerProtocol->AddInternal(EPixelStreaming2FromStreamerMessage::Command, 2);
+		FromStreamerProtocol->AddInternal(EPixelStreaming2FromStreamerMessage::FreezeFrame, 3);
+		FromStreamerProtocol->AddInternal(EPixelStreaming2FromStreamerMessage::UnfreezeFrame, 4);
+		FromStreamerProtocol->AddInternal(EPixelStreaming2FromStreamerMessage::VideoEncoderAvgQP, 5);
+		FromStreamerProtocol->AddInternal(EPixelStreaming2FromStreamerMessage::LatencyTest, 6);
+		FromStreamerProtocol->AddInternal(EPixelStreaming2FromStreamerMessage::InitialSettings, 7);
+		FromStreamerProtocol->AddInternal(EPixelStreaming2FromStreamerMessage::FileExtension, 8);
+		FromStreamerProtocol->AddInternal(EPixelStreaming2FromStreamerMessage::FileMimeType, 9);
+		FromStreamerProtocol->AddInternal(EPixelStreaming2FromStreamerMessage::FileContents, 10);
+		FromStreamerProtocol->AddInternal(EPixelStreaming2FromStreamerMessage::TestEcho, 11);
+		FromStreamerProtocol->AddInternal(EPixelStreaming2FromStreamerMessage::InputControlOwnership, 12);
+		FromStreamerProtocol->AddInternal(EPixelStreaming2FromStreamerMessage::GamepadResponse, 13);
+		FromStreamerProtocol->AddInternal(EPixelStreaming2FromStreamerMessage::Protocol, 255);
+		FromStreamerProtocol->AddInternal(EPixelStreaming2FromStreamerMessage::Multiplexed, 199);
+
+		return FromStreamerProtocol;
+	}
+
+} // namespace UE::PixelStreaming2Input
\ No newline at end of file
diff --git a/Engine/Plugins/Media/PixelStreaming2/Source/PixelStreaming2Input/Private/PixelStreaming2InputHandler.cpp b/Engine/Plugins/Media/PixelStreaming2/Source/PixelStreaming2Input/Private/PixelStreaming2InputHandler.cpp
new file mode 100644
--- /dev/null
+++ b/Engine/Plugins/Media/PixelStreaming2/Source/PixelStreaming2Input/Private/PixelStreaming2InputHandler.cpp
@@ -0,0 +1,16 @@
+// Copyright Epic Games, Inc. All Rights Reserved.
+
+#include "IPixelStreaming2InputHandler.h"
+
+#include "InputDevice.h"
+
+IPixelStreaming2InputHandler::IPixelStreaming2InputHandler()
+{
+    // Register this input handler with the module's input device so that it's ticked
+    UE::PixelStreaming2Input::FInputDevice::GetInputDevice()->AddInputHandler(this);
+}
+
+IPixelStreaming2InputHandler::~IPixelStreaming2InputHandler()
+{
+    UE::PixelStreaming2Input::FInputDevice::GetInputDevice()->RemoveInputHandler(this);
+}
\ No newline at end of file
diff --git a/Engine/Plugins/Media/PixelStreaming2/Source/PixelStreaming2Input/Public/IPixelStreaming2InputHandler.h b/Engine/Plugins/Media/PixelStreaming2/Source/PixelStreaming2Input/Public/IPixelStreaming2InputHandler.h
--- a/Engine/Plugins/Media/PixelStreaming2/Source/PixelStreaming2Input/Public/IPixelStreaming2InputHandler.h
+++ b/Engine/Plugins/Media/PixelStreaming2/Source/PixelStreaming2Input/Public/IPixelStreaming2InputHandler.h
@@ -8,6 +8,8 @@
 #include "PixelStreaming2HMDEnums.h"
 #include "Widgets/SViewport.h"
 
+#define UE_API PIXELSTREAMING2INPUT_API
+
 /**
  * The IPixelStreaming2InputHandler, used to handle input from a remote peer and pass it to UE accordingly. Setting the target viewport allows for
  * scaling of input from browser to application, and setting the target window ensure that if windows are tiled (eg editor)
@@ -16,6 +18,9 @@
 class IPixelStreaming2InputHandler : public IInputDevice
 {
 public:
+	UE_API IPixelStreaming2InputHandler();
+	UE_API virtual ~IPixelStreaming2InputHandler();
+
 	/**
 	 * @brief Handle the message from the WebRTC data channel.
 	 * @param SourceId A source ID for this message
@@ -410,3 +415,5 @@ public:
 	 */
 	virtual bool OnXRSystem(EPixelStreaming2XRSystem System) = 0;
 };
+
+#undef UE_API
diff --git a/Engine/Plugins/Media/PixelStreaming2/Source/PixelStreaming2Input/Public/PixelStreaming2DefaultDataProtocol.h b/Engine/Plugins/Media/PixelStreaming2/Source/PixelStreaming2Input/Public/PixelStreaming2DefaultDataProtocol.h
new file mode 100644
--- /dev/null
+++ b/Engine/Plugins/Media/PixelStreaming2/Source/PixelStreaming2Input/Public/PixelStreaming2DefaultDataProtocol.h
@@ -0,0 +1,21 @@
+// Copyright Epic Games, Inc. All Rights Reserved.
+
+#pragma once
+
+#include "Templates/SharedPointer.h"
+#include "IPixelStreaming2DataProtocol.h"
+
+namespace UE::PixelStreaming2Input
+{
+
+	/**
+	 * @return The default "ToStreamer" data protocol.
+	 */
+	PIXELSTREAMING2INPUT_API TSharedPtr<IPixelStreaming2DataProtocol> GetDefaultToStreamerProtocol();
+
+	/**
+	 * @return The default "FromStreamer" data protocol.
+	 */
+	PIXELSTREAMING2INPUT_API TSharedPtr<IPixelStreaming2DataProtocol> GetDefaultFromStreamerProtocol();
+
+} // namespace UE::PixelStreaming2Input
\ No newline at end of file
diff --git a/Engine/Plugins/Media/PixelStreaming2/Source/PixelStreaming2Input/Public/PixelStreaming2InputStructures.h b/Engine/Plugins/Media/PixelStreaming2/Source/PixelStreaming2Input/Public/PixelStreaming2InputStructures.h
new file mode 100644
--- /dev/null
+++ b/Engine/Plugins/Media/PixelStreaming2/Source/PixelStreaming2Input/Public/PixelStreaming2InputStructures.h
@@ -0,0 +1,32 @@
+// Copyright Epic Games, Inc. All Rights Reserved.
+
+#pragma once
+
+#include "CoreMinimal.h"
+#include "Serialization/BufferArchive.h"
+#include "Serialization/MemoryReader.h"
+
+namespace UE::PixelStreaming2Input
+{
+	template <typename IntegerSequence>
+	struct TPayloadParser;
+
+	template <uint32... Indices>
+	struct TPayloadParser<TIntegerSequence<uint32, Indices...>>
+	{
+		template <typename TupleType>
+		static void Parse(TupleType&& Tuple, FArchive& Ar)
+		{
+			([](auto& Arg, FArchive& Ar) { Ar << Arg; }(Forward<TupleType>(Tuple).template Get<Indices>(), Ar), ...);
+		}
+	};
+
+	template <typename... Types>
+	struct TPayload : public TTuple<std::decay_t<Types>...>
+	{
+		TPayload(FArchive& Ar)
+		{
+			TPayloadParser<TMakeIntegerSequence<uint32, sizeof...(Types)>>::Parse(*this, Ar);
+		}
+	};
+} // namespace UE::PixelStreaming2Input
diff --git a/Engine/Plugins/Media/PixelStreaming2/Source/PixelStreaming2Input/Public/PixelStreaming2JavaScriptKeyCodes.inl b/Engine/Plugins/Media/PixelStreaming2/Source/PixelStreaming2Input/Public/PixelStreaming2JavaScriptKeyCodes.inl
new file mode 100644
--- /dev/null
+++ b/Engine/Plugins/Media/PixelStreaming2/Source/PixelStreaming2Input/Public/PixelStreaming2JavaScriptKeyCodes.inl
@@ -0,0 +1,280 @@
+// Copyright Epic Games, Inc. All Rights Reserved.
+
+#pragma once
+
+#include "InputCoreTypes.h"
+
+namespace UE::PixelStreaming2Input 
+{
+
+/**
+ * An array which will convert JavaScript key codes into FKeys. The JavaScript
+ * key code should be used as an index into the array to perform the lookup.
+ * https://www.cambiaresearch.com/articles/15/javascript-char-codes-key-codes
+ * http://keycode.info/
+ * Most key codes use the standard values given in the links above, but some
+ * have been choosen specially for 'right' versions of the same key.
+ */
+
+	static const FKey* JavaScriptKeyCodeToFKey[256] =
+	{
+		/*   0 */ &EKeys::Invalid,
+		/*   1 */ &EKeys::Invalid,
+		/*   2 */ &EKeys::Invalid,
+		/*   3 */ &EKeys::Invalid,
+		/*   4 */ &EKeys::Invalid,
+		/*   5 */ &EKeys::Invalid,
+		/*   6 */ &EKeys::Invalid,
+		/*   7 */ &EKeys::Invalid,
+		/*   8 */ &EKeys::BackSpace,
+		/*   9 */ &EKeys::Tab,
+		/*  10 */ &EKeys::Invalid,
+		/*  11 */ &EKeys::Invalid,
+		/*  12 */ /*&EKeys::Clear*/ &EKeys::Invalid,
+		/*  13 */ &EKeys::Enter,
+		/*  14 */ &EKeys::Invalid,
+		/*  15 */ &EKeys::Invalid,
+		/*  16 */ &EKeys::LeftShift,
+		/*  17 */ &EKeys::LeftControl,
+		/*  18 */ &EKeys::LeftAlt,
+		/*  19 */ &EKeys::Pause,
+		/*  20 */ &EKeys::CapsLock,
+		/*  21 */ &EKeys::Invalid,
+		/*  22 */ &EKeys::Invalid,
+		/*  23 */ &EKeys::Invalid,
+		/*  24 */ &EKeys::Invalid,
+		/*  25 */ &EKeys::Invalid,
+		/*  26 */ &EKeys::Invalid,
+		/*  27 */ &EKeys::Escape,
+		/*  28 */ &EKeys::Invalid,
+		/*  29 */ &EKeys::Invalid,
+		/*  30 */ &EKeys::Invalid,
+		/*  31 */ &EKeys::Invalid,
+		/*  32 */ &EKeys::SpaceBar,
+		/*  33 */ &EKeys::PageUp,
+		/*  34 */ &EKeys::PageDown,
+		/*  35 */ &EKeys::End,
+		/*  36 */ &EKeys::Home,
+		/*  37 */ &EKeys::Left,
+		/*  38 */ &EKeys::Up,
+		/*  39 */ &EKeys::Right,
+		/*  40 */ &EKeys::Down,
+		/*  41 */ &EKeys::Invalid,
+		/*  42 */ &EKeys::Invalid,
+		/*  43 */ &EKeys::Invalid,
+		/*  44 */ /*&EKeys::PrintScreen*/ &EKeys::Invalid,
+		/*  45 */ &EKeys::Insert,
+		/*  46 */ &EKeys::Delete,
+		/*  47 */ &EKeys::Invalid,
+		/*  48 */ &EKeys::Zero,
+		/*  49 */ &EKeys::One,
+		/*  50 */ &EKeys::Two,
+		/*  51 */ &EKeys::Three,
+		/*  52 */ &EKeys::Four,
+		/*  53 */ &EKeys::Five,
+		/*  54 */ &EKeys::Six,
+		/*  55 */ &EKeys::Seven,
+		/*  56 */ &EKeys::Eight,
+		/*  57 */ &EKeys::Nine,
+		/*  58 */ &EKeys::Invalid,
+		/*  59 */ &EKeys::Invalid,
+		/*  60 */ &EKeys::Invalid,
+		/*  61 */ &EKeys::Invalid,
+		/*  62 */ &EKeys::Invalid,
+		/*  63 */ &EKeys::Invalid,
+		/*  64 */ &EKeys::Invalid,
+		/*  65 */ &EKeys::A,
+		/*  66 */ &EKeys::B,
+		/*  67 */ &EKeys::C,
+		/*  68 */ &EKeys::D,
+		/*  69 */ &EKeys::E,
+		/*  70 */ &EKeys::F,
+		/*  71 */ &EKeys::G,
+		/*  72 */ &EKeys::H,
+		/*  73 */ &EKeys::I,
+		/*  74 */ &EKeys::J,
+		/*  75 */ &EKeys::K,
+		/*  76 */ &EKeys::L,
+		/*  77 */ &EKeys::M,
+		/*  78 */ &EKeys::N,
+		/*  79 */ &EKeys::O,
+		/*  80 */ &EKeys::P,
+		/*  81 */ &EKeys::Q,
+		/*  82 */ &EKeys::R,
+		/*  83 */ &EKeys::S,
+		/*  84 */ &EKeys::T,
+		/*  85 */ &EKeys::U,
+		/*  86 */ &EKeys::V,
+		/*  87 */ &EKeys::W,
+		/*  88 */ &EKeys::X,
+		/*  89 */ &EKeys::Y,
+		/*  90 */ &EKeys::Z,
+		/*  91 */ /*&EKeys::LeftWindowKey*/ &EKeys::Invalid,
+		/*  92 */ /*&EKeys::RightWindowKey*/ &EKeys::Invalid,
+		/*  93 */ /*&EKeys::SelectKey*/ &EKeys::Invalid,
+		/*  94 */ &EKeys::Invalid,
+		/*  95 */ &EKeys::Invalid,
+		/*  96 */ &EKeys::NumPadZero,
+		/*  97 */ &EKeys::NumPadOne,
+		/*  98 */ &EKeys::NumPadTwo,
+		/*  99 */ &EKeys::NumPadThree,
+		/* 100 */ &EKeys::NumPadFour,
+		/* 101 */ &EKeys::NumPadFive,
+		/* 102 */ &EKeys::NumPadSix,
+		/* 103 */ &EKeys::NumPadSeven,
+		/* 104 */ &EKeys::NumPadEight,
+		/* 105 */ &EKeys::NumPadNine,
+		/* 106 */ &EKeys::Multiply,
+		/* 107 */ &EKeys::Add,
+		/* 108 */ &EKeys::Invalid,
+		/* 109 */ &EKeys::Subtract,
+		/* 110 */ &EKeys::Decimal,
+		/* 111 */ &EKeys::Divide,
+		/* 112 */ &EKeys::F1,
+		/* 113 */ &EKeys::F2,
+		/* 114 */ &EKeys::F3,
+		/* 115 */ &EKeys::F4,
+		/* 116 */ &EKeys::F5,
+		/* 117 */ &EKeys::F6,
+		/* 118 */ &EKeys::F7,
+		/* 119 */ &EKeys::F8,
+		/* 120 */ &EKeys::F9,
+		/* 121 */ &EKeys::F10,
+		/* 122 */ &EKeys::F11,
+		/* 123 */ &EKeys::F12,
+		/* 124 */ &EKeys::Invalid,
+		/* 125 */ &EKeys::Invalid,
+		/* 126 */ &EKeys::Invalid,
+		/* 127 */ &EKeys::Invalid,
+		/* 128 */ &EKeys::Invalid,
+		/* 129 */ &EKeys::Invalid,
+		/* 130 */ &EKeys::Invalid,
+		/* 131 */ &EKeys::Invalid,
+		/* 132 */ &EKeys::Invalid,
+		/* 133 */ &EKeys::Invalid,
+		/* 134 */ &EKeys::Invalid,
+		/* 135 */ &EKeys::Invalid,
+		/* 136 */ &EKeys::Invalid,
+		/* 137 */ &EKeys::Invalid,
+		/* 138 */ &EKeys::Invalid,
+		/* 139 */ &EKeys::Invalid,
+		/* 140 */ &EKeys::Invalid,
+		/* 141 */ &EKeys::Invalid,
+		/* 142 */ &EKeys::Invalid,
+		/* 143 */ &EKeys::Invalid,
+		/* 144 */ &EKeys::NumLock,
+		/* 145 */ &EKeys::ScrollLock,
+		/* 146 */ &EKeys::Invalid,
+		/* 147 */ &EKeys::Invalid,
+		/* 148 */ &EKeys::Invalid,
+		/* 149 */ &EKeys::Invalid,
+		/* 150 */ &EKeys::Invalid,
+		/* 151 */ &EKeys::Invalid,
+		/* 152 */ &EKeys::Invalid,
+		/* 153 */ &EKeys::Invalid,
+		/* 154 */ &EKeys::Invalid,
+		/* 155 */ &EKeys::Invalid,
+		/* 156 */ &EKeys::Invalid,
+		/* 157 */ &EKeys::Invalid,
+		/* 158 */ &EKeys::Invalid,
+		/* 159 */ &EKeys::Invalid,
+		/* 160 */ &EKeys::Invalid,
+		/* 161 */ &EKeys::Invalid,
+		/* 162 */ &EKeys::Invalid,
+		/* 163 */ &EKeys::Invalid,
+		/* 164 */ &EKeys::Invalid,
+		/* 165 */ &EKeys::Invalid,
+		/* 166 */ &EKeys::Invalid,
+		/* 167 */ &EKeys::Invalid,
+		/* 168 */ &EKeys::Invalid,
+		/* 169 */ &EKeys::Invalid,
+		/* 170 */ &EKeys::Invalid,
+		/* 171 */ &EKeys::Invalid,
+		/* 172 */ &EKeys::Invalid,
+		/* 173 */ &EKeys::Invalid,
+		/* 174 */ &EKeys::Invalid,
+		/* 175 */ &EKeys::Invalid,
+		/* 176 */ &EKeys::Invalid,
+		/* 177 */ &EKeys::Invalid,
+		/* 178 */ &EKeys::Invalid,
+		/* 179 */ &EKeys::Invalid,
+		/* 180 */ &EKeys::Invalid,
+		/* 181 */ &EKeys::Invalid,
+		/* 182 */ &EKeys::Invalid,
+		/* 183 */ &EKeys::Invalid,
+		/* 184 */ &EKeys::Invalid,
+		/* 185 */ &EKeys::Invalid,
+		/* 186 */ &EKeys::Semicolon,
+		/* 187 */ &EKeys::Equals,
+		/* 188 */ &EKeys::Comma,
+		/* 189 */ &EKeys::Hyphen,
+		/* 190 */ &EKeys::Period,
+		/* 191 */ &EKeys::Slash,
+		/* 192 */ &EKeys::Tilde,
+		/* 193 */ &EKeys::Invalid,
+		/* 194 */ &EKeys::Invalid,
+		/* 195 */ &EKeys::Invalid,
+		/* 196 */ &EKeys::Invalid,
+		/* 197 */ &EKeys::Invalid,
+		/* 198 */ &EKeys::Invalid,
+		/* 199 */ &EKeys::Invalid,
+		/* 200 */ &EKeys::Invalid,
+		/* 201 */ &EKeys::Invalid,
+		/* 202 */ &EKeys::Invalid,
+		/* 203 */ &EKeys::Invalid,
+		/* 204 */ &EKeys::Invalid,
+		/* 205 */ &EKeys::Invalid,
+		/* 206 */ &EKeys::Invalid,
+		/* 207 */ &EKeys::Invalid,
+		/* 208 */ &EKeys::Invalid,
+		/* 209 */ &EKeys::Invalid,
+		/* 210 */ &EKeys::Invalid,
+		/* 211 */ &EKeys::Invalid,
+		/* 212 */ &EKeys::Invalid,
+		/* 213 */ &EKeys::Invalid,
+		/* 214 */ &EKeys::Invalid,
+		/* 215 */ &EKeys::Invalid,
+		/* 216 */ &EKeys::Invalid,
+		/* 217 */ &EKeys::Invalid,
+		/* 218 */ &EKeys::Invalid,
+		/* 219 */ &EKeys::LeftBracket,
+		/* 220 */ &EKeys::Backslash,
+		/* 221 */ &EKeys::RightBracket,
+		/* 222 */ &EKeys::Apostrophe,
+		/* 223 */ &EKeys::Tilde,
+		/* 224 */ &EKeys::Invalid,
+		/* 225 */ &EKeys::Invalid,
+		/* 226 */ &EKeys::Invalid,
+		/* 227 */ &EKeys::Invalid,
+		/* 228 */ &EKeys::Invalid,
+		/* 229 */ &EKeys::Invalid,
+		/* 230 */ &EKeys::Invalid,
+		/* 231 */ &EKeys::Invalid,
+		/* 232 */ &EKeys::Invalid,
+		/* 233 */ &EKeys::Invalid,
+		/* 234 */ &EKeys::Invalid,
+		/* 235 */ &EKeys::Invalid,
+		/* 236 */ &EKeys::Invalid,
+		/* 237 */ &EKeys::Invalid,
+		/* 238 */ &EKeys::Invalid,
+		/* 239 */ &EKeys::Invalid,
+		/* 240 */ &EKeys::Invalid,
+		/* 241 */ &EKeys::Invalid,
+		/* 242 */ &EKeys::Invalid,
+		/* 243 */ &EKeys::Invalid,
+		/* 244 */ &EKeys::Invalid,
+		/* 245 */ &EKeys::Invalid,
+		/* 246 */ &EKeys::Invalid,
+		/* 247 */ &EKeys::Invalid,
+		/* 248 */ &EKeys::Invalid,
+		/* 249 */ &EKeys::Invalid,
+		/* 250 */ &EKeys::Invalid,
+		/* 251 */ &EKeys::Invalid,
+		/* 252 */ &EKeys::Invalid,
+
+		// Special key codes.
+		/* 253 */ &EKeys::RightShift,
+		/* 254 */ &EKeys::RightControl,
+		/* 255 */ &EKeys::RightAlt
+	};
+} // namespace UE::PixelStreaming2Input
\ No newline at end of file
diff --git a/Engine/Plugins/Media/PixelStreaming2/Source/PixelStreaming2RTC/Internal/EpicRtcDataTrack.h b/Engine/Plugins/Media/PixelStreaming2/Source/PixelStreaming2RTC/Internal/EpicRtcDataTrack.h
--- a/Engine/Plugins/Media/PixelStreaming2/Source/PixelStreaming2RTC/Internal/EpicRtcDataTrack.h
+++ b/Engine/Plugins/Media/PixelStreaming2/Source/PixelStreaming2RTC/Internal/EpicRtcDataTrack.h
@@ -102,6 +102,11 @@ namespace UE::PixelStreaming2
 
 		virtual ~FEpicRtcDataTrack() = default;
 
+		void SetDataProtocol(TWeakPtr<IPixelStreaming2DataProtocol> InDataProtocol)
+		{
+			WeakDataProtocol = InDataProtocol;
+		}
+
 		/**
 		 * Sends a series of arguments to the data channel with the given type.
 		 * @param MessageType The name of the message you want to send. This message name must be registered IPixelStreaming2InputHandler::GetFromStreamerProtocol().
diff --git a/Engine/Plugins/Media/PixelStreaming2/Source/PixelStreaming2RTC/Private/EpicRtcStreamer.cpp b/Engine/Plugins/Media/PixelStreaming2/Source/PixelStreaming2RTC/Private/EpicRtcStreamer.cpp
--- a/Engine/Plugins/Media/PixelStreaming2/Source/PixelStreaming2RTC/Private/EpicRtcStreamer.cpp
+++ b/Engine/Plugins/Media/PixelStreaming2/Source/PixelStreaming2RTC/Private/EpicRtcStreamer.cpp
@@ -21,19 +21,12 @@
 {
 	FEpicRtcStreamer::FEpicRtcStreamer(const FString& InStreamerId, TRefCountPtr<EpicRtcConferenceInterface> Conference)
 		: StreamerId(InStreamerId)
-		, InputHandler(FRTCInputHandler::Create())
 		, Participants(new TThreadSafeMap<FString, TSharedPtr<FPlayerContext>>())
 		, VideoCapturer(FEpicRtcVideoCapturer::Create())
 		, VideoSourceGroup(FVideoSourceGroup::Create(VideoCapturer))
-		, FreezeFrame(FFreezeFrame::Create(Participants, VideoCapturer, InputHandler))
 		, ReconnectTimer(MakeShared<FStreamerReconnectTimer>())
 		, EpicRtcConference(MoveTemp(Conference))
 	{
-		InputHandler->SetElevatedCheck([this](FString PlayerId) {
-			return GetEnumFromCVar<EInputControllerMode>(UPixelStreaming2PluginSettings::CVarInputController) == EInputControllerMode::Any
-				|| InputControllingId == INVALID_PLAYER_ID
-				|| PlayerId == InputControllingId;
-		});
 	}
 
 	FEpicRtcStreamer::~FEpicRtcStreamer()
@@ -64,94 +57,6 @@
 		VideoTrackObserverFactory = MakeRefCount<FEpicRtcVideoTrackObserverFactory>(TObserver<IPixelStreaming2VideoTrackObserver>(WeakStreamer));
 		DataTrackObserverFactory = MakeRefCount<FEpicRtcDataTrackObserverFactory>(TObserver<IPixelStreaming2DataTrackObserver>(WeakStreamer));
 
-		InputHandler->GetToStreamerProtocol()->OnProtocolUpdated().AddSP(Streamer.ToSharedRef(), &FEpicRtcStreamer::OnProtocolUpdated);
-		InputHandler->GetFromStreamerProtocol()->OnProtocolUpdated().AddSP(Streamer.ToSharedRef(), &FEpicRtcStreamer::OnProtocolUpdated);
-
-		// Set Encoder.MinQP Legacy CVar
-		InputHandler->SetCommandHandler(TEXT("Encoder.MinQP"), [](FString PlayerId, FString Descriptor, FString MinQPString) {
-			int MinQP = FCString::Atoi(*MinQPString);
-			UPixelStreaming2PluginSettings::CVarEncoderMaxQuality->SetWithCurrentPriority(FMath::RoundToFloat(100.0f * (1.0f - (FMath::Clamp<int32>(MinQP, 0, 51) / 51.0f))));
-		});
-
-		// Set Encoder.MaxQP Legacy CVar
-		InputHandler->SetCommandHandler(TEXT("Encoder.MaxQP"), [](FString PlayerId, FString Descriptor, FString MaxQPString) {
-			int MaxQP = FCString::Atoi(*MaxQPString);
-			UPixelStreaming2PluginSettings::CVarEncoderMinQuality->SetWithCurrentPriority(FMath::RoundToFloat(100.0f * (1.0f - (FMath::Clamp<int32>(MaxQP, 0, 51) / 51.0f))));
-		});
-
-		// Set Encoder.MinQuality CVar
-		InputHandler->SetCommandHandler(TEXT("Encoder.MinQuality"), [](FString PlayerId, FString Descriptor, FString MinQualityString) {
-			int MinQuality = FCString::Atoi(*MinQualityString);
-			UPixelStreaming2PluginSettings::CVarEncoderMinQuality->SetWithCurrentPriority(FMath::Clamp<int32>(MinQuality, 0, 100));
-		});
-
-		// Set Encoder.MaxQuality CVar
-		InputHandler->SetCommandHandler(TEXT("Encoder.MaxQuality"), [](FString PlayerId, FString Descriptor, FString MaxQualityString) {
-			int MaxQuality = FCString::Atoi(*MaxQualityString);
-			UPixelStreaming2PluginSettings::CVarEncoderMaxQuality->SetWithCurrentPriority(FMath::Clamp<int32>(MaxQuality, 0, 100));
-		});
-
-		// Set WebRTC max FPS
-		InputHandler->SetCommandHandler(TEXT("WebRTC.Fps"), [](FString PlayerId, FString Descriptor, FString FPSString) {
-			int FPS = FCString::Atoi(*FPSString);
-			UPixelStreaming2PluginSettings::CVarWebRTCFps->SetWithCurrentPriority(FPS);
-		});
-
-		// Set MinBitrate
-		InputHandler->SetCommandHandler(TEXT("WebRTC.MinBitrate"), [WeakStreamer](FString PlayerId, FString Descriptor, FString MinBitrateString) {
-			// This check pattern is kind of verbose, but as the messages are enqueued on a different thread to where they were added
-			// we need to make sure that both the streamer and input handler are still alive when we process the command
-			TSharedPtr<FEpicRtcStreamer> Streamer = WeakStreamer.Pin();
-			if (!Streamer)
-			{
-				return;
-			}
-			TSharedPtr<IPixelStreaming2InputHandler> InputHandler = Streamer->GetInputHandler().Pin();
-			if (!InputHandler)
-			{
-				return;
-			}
-
-			if (InputHandler->IsElevated(PlayerId))
-			{
-				int MinBitrate = FCString::Atoi(*MinBitrateString);
-				UPixelStreaming2PluginSettings::CVarWebRTCMinBitrate->SetWithCurrentPriority(MinBitrate);
-			}
-		});
-
-		// Set MaxBitrate
-		InputHandler->SetCommandHandler(TEXT("WebRTC.MaxBitrate"), [WeakStreamer](FString PlayerId, FString Descriptor, FString MaxBitrateString) {
-			TSharedPtr<FEpicRtcStreamer> Streamer = WeakStreamer.Pin();
-			if (!Streamer)
-			{
-				return;
-			}
-			TSharedPtr<IPixelStreaming2InputHandler> InputHandler = Streamer->GetInputHandler().Pin();
-			if (!InputHandler)
-			{
-				return;
-			}
-
-			if (InputHandler->IsElevated(PlayerId))
-			{
-				int MaxBitrate = FCString::Atoi(*MaxBitrateString);
-				UPixelStreaming2PluginSettings::CVarWebRTCMaxBitrate->SetWithCurrentPriority(MaxBitrate);
-			}
-		});
-
-		InputHandler->RegisterMessageHandler(EPixelStreaming2ToStreamerMessage::UIInteraction, [WeakStreamer](FString PlayerId, FMemoryReader Ar) {
-			TSharedPtr<FEpicRtcStreamer> Streamer = WeakStreamer.Pin();
-			if (!Streamer)
-			{
-				return;
-			}
-			Streamer->OnUIInteraction(Ar);
-		});
-
-		// Handle special cases when the InputHandler itself wants to send a message out to all the peers.
-		// Some special cases include when virtual gamepads are connected and a controller id needs to be transmitted.
-		InputHandler->OnSendMessage.AddSP(Streamer.ToSharedRef(), &FEpicRtcStreamer::OnSendMessage);
-
 		VideoCapturer->OnFrameCapturerCreated.AddSP(Streamer.ToSharedRef(), &FEpicRtcStreamer::OnFrameCapturerCreated);
 
 		if (UPixelStreaming2Delegates* Delegates = UPixelStreaming2Delegates::Get())
@@ -456,14 +361,20 @@
 
 	void FEpicRtcStreamer::FreezeStream(UTexture2D* Texture)
 	{
-		FreezeFrame->StartFreeze(Texture);
+		if (FreezeFrame)
+		{
+			FreezeFrame->StartFreeze(Texture);
+		}
 	}
 
 	void FEpicRtcStreamer::UnfreezeStream()
 	{
 		// Force a keyframe so when stream unfreezes if player has never received a frame before they can still connect.
 		ForceKeyFrame();
-		FreezeFrame->StopFreeze();
+		if (FreezeFrame)
+		{
+			FreezeFrame->StopFreeze();
+		}
 	}
 
 	void FEpicRtcStreamer::SendAllPlayersMessage(FString MessageType, const FString& Descriptor)
@@ -526,7 +437,10 @@
 			ParticipantInterface = Participant->ParticipantInterface;
 		}
 
-		ParticipantInterface->Kick();
+		if (ParticipantInterface)
+		{
+			ParticipantInterface->Kick();
+		}
 	}
 
 	TArray<FString> FEpicRtcStreamer::GetConnectedPlayers()
@@ -537,6 +451,126 @@
 		return ConnectedParticipantIds.Array();
 	}
 
+	void FEpicRtcStreamer::SetInputHandler(TSharedPtr<IPixelStreaming2InputHandler> Handler)
+	{
+		InputHandler = Handler;
+
+		if (!InputHandler)
+		{
+			return;
+		}
+
+		// Re-create the freeze frame manager with the new input handler
+		FreezeFrame = FFreezeFrame::Create(Participants, VideoCapturer, InputHandler);
+
+		InputHandler->SetElevatedCheck([this](FString PlayerId) {
+			return GetEnumFromCVar<EInputControllerMode>(UPixelStreaming2PluginSettings::CVarInputController) == EInputControllerMode::Any
+				|| InputControllingId == INVALID_PLAYER_ID
+				|| PlayerId == InputControllingId;
+		});
+
+		TSharedPtr<FEpicRtcStreamer> Streamer = AsShared();
+		TWeakPtr<FEpicRtcStreamer>	 WeakStreamer = Streamer;
+
+		InputHandler->GetToStreamerProtocol()->OnProtocolUpdated().AddSP(Streamer.ToSharedRef(), &FEpicRtcStreamer::OnProtocolUpdated);
+		InputHandler->GetFromStreamerProtocol()->OnProtocolUpdated().AddSP(Streamer.ToSharedRef(), &FEpicRtcStreamer::OnProtocolUpdated);
+
+		// Set Encoder.MinQP Legacy CVar
+		InputHandler->SetCommandHandler(TEXT("Encoder.MinQP"), [](FString PlayerId, FString Descriptor, FString MinQPString) {
+			int MinQP = FCString::Atoi(*MinQPString);
+			UPixelStreaming2PluginSettings::CVarEncoderMaxQuality->SetWithCurrentPriority(FMath::RoundToFloat(100.0f * (1.0f - (FMath::Clamp<int32>(MinQP, 0, 51) / 51.0f))));
+		});
+
+		// Set Encoder.MaxQP Legacy CVar
+		InputHandler->SetCommandHandler(TEXT("Encoder.MaxQP"), [](FString PlayerId, FString Descriptor, FString MaxQPString) {
+			int MaxQP = FCString::Atoi(*MaxQPString);
+			UPixelStreaming2PluginSettings::CVarEncoderMinQuality->SetWithCurrentPriority(FMath::RoundToFloat(100.0f * (1.0f - (FMath::Clamp<int32>(MaxQP, 0, 51) / 51.0f))));
+		});
+
+		// Set Encoder.MinQuality CVar
+		InputHandler->SetCommandHandler(TEXT("Encoder.MinQuality"), [](FString PlayerId, FString Descriptor, FString MinQualityString) {
+			int MinQuality = FCString::Atoi(*MinQualityString);
+			UPixelStreaming2PluginSettings::CVarEncoderMinQuality->SetWithCurrentPriority(FMath::Clamp<int32>(MinQuality, 0, 100));
+		});
+
+		// Set Encoder.MaxQuality CVar
+		InputHandler->SetCommandHandler(TEXT("Encoder.MaxQuality"), [](FString PlayerId, FString Descriptor, FString MaxQualityString) {
+			int MaxQuality = FCString::Atoi(*MaxQualityString);
+			UPixelStreaming2PluginSettings::CVarEncoderMaxQuality->SetWithCurrentPriority(FMath::Clamp<int32>(MaxQuality, 0, 100));
+		});
+
+		// Set WebRTC max FPS
+		InputHandler->SetCommandHandler(TEXT("WebRTC.Fps"), [](FString PlayerId, FString Descriptor, FString FPSString) {
+			int FPS = FCString::Atoi(*FPSString);
+			UPixelStreaming2PluginSettings::CVarWebRTCFps->SetWithCurrentPriority(FPS);
+		});
+
+		// Set MinBitrate
+		InputHandler->SetCommandHandler(TEXT("WebRTC.MinBitrate"), [WeakStreamer](FString PlayerId, FString Descriptor, FString MinBitrateString) {
+			// This check pattern is kind of verbose, but as the messages are enqueued on a different thread to where they were added
+			// we need to make sure that both the streamer and input handler are still alive when we process the command
+			TSharedPtr<FEpicRtcStreamer> Streamer = WeakStreamer.Pin();
+			if (!Streamer)
+			{
+				return;
+			}
+			TSharedPtr<IPixelStreaming2InputHandler> InputHandler = Streamer->GetInputHandler().Pin();
+			if (!InputHandler)
+			{
+				return;
+			}
+
+			if (InputHandler->IsElevated(PlayerId))
+			{
+				int MinBitrate = FCString::Atoi(*MinBitrateString);
+				UPixelStreaming2PluginSettings::CVarWebRTCMinBitrate->SetWithCurrentPriority(MinBitrate);
+			}
+		});
+
+		// Set MaxBitrate
+		InputHandler->SetCommandHandler(TEXT("WebRTC.MaxBitrate"), [WeakStreamer](FString PlayerId, FString Descriptor, FString MaxBitrateString) {
+			TSharedPtr<FEpicRtcStreamer> Streamer = WeakStreamer.Pin();
+			if (!Streamer)
+			{
+				return;
+			}
+			TSharedPtr<IPixelStreaming2InputHandler> InputHandler = Streamer->GetInputHandler().Pin();
+			if (!InputHandler)
+			{
+				return;
+			}
+
+			if (InputHandler->IsElevated(PlayerId))
+			{
+				int MaxBitrate = FCString::Atoi(*MaxBitrateString);
+				UPixelStreaming2PluginSettings::CVarWebRTCMaxBitrate->SetWithCurrentPriority(MaxBitrate);
+			}
+		});
+
+		InputHandler->RegisterMessageHandler(EPixelStreaming2ToStreamerMessage::UIInteraction, [WeakStreamer](FString PlayerId, FMemoryReader Ar) {
+			TSharedPtr<FEpicRtcStreamer> Streamer = WeakStreamer.Pin();
+			if (!Streamer)
+			{
+				return;
+			}
+			Streamer->OnUIInteraction(Ar);
+		});
+
+		// Handle special cases when the InputHandler itself wants to send a message out to all the peers.
+		// Some special cases include when virtual gamepads are connected and a controller id needs to be transmitted.
+		InputHandler->OnSendMessage.AddSP(Streamer.ToSharedRef(), &FEpicRtcStreamer::OnSendMessage);
+
+		// We need to update the protocol for each existing participant's data track
+		TSharedPtr<IPixelStreaming2InputHandler> SharedInputHandler = InputHandler;
+		Participants->Apply([SharedInputHandler](FString PlayerId, TSharedPtr<FPlayerContext> Participant) {
+			TSharedPtr<IPixelStreaming2DataProtocol> Protocol = SharedInputHandler->GetFromStreamerProtocol();
+			if (Participant->DataTrack && Protocol)
+			{
+				Participant->DataTrack->SetDataProtocol(Protocol);
+			}
+		});
+	}
+
 	TWeakPtr<IPixelStreaming2AudioSink> FEpicRtcStreamer::GetPeerAudioSink(FString PlayerId)
 	{
 		TWeakPtr<IPixelStreaming2AudioSink> Result = nullptr;
@@ -784,10 +818,14 @@
 
 		// When data channel is open
 		SendProtocol(PlayerId);
-		// Try to send cached freeze frame (if we have one)
-		FreezeFrame->SendCachedFreezeFrameTo(PlayerId);
 		SendInitialSettings(PlayerId);
 		SendPeerControllerMessages(PlayerId);
+		
+		// Try to send cached freeze frame (if we have one)
+		if (FreezeFrame)
+		{
+			FreezeFrame->SendCachedFreezeFrameTo(PlayerId);
+		}
 	}
 
 	void FEpicRtcStreamer::OnDataTrackClosed(FString PlayerId)
@@ -966,6 +1004,12 @@
 
 	void FEpicRtcStreamer::SendProtocol(FString PlayerId) const
 	{
+		if (!InputHandler)
+		{
+			UE_LOGFMT(LogPixelStreaming2RTC, Warning, "Cannot send protocol when input handler is not valid!");
+			return;
+		}
+
 		const TArray<TSharedPtr<IPixelStreaming2DataProtocol>> Protocols = { InputHandler->GetToStreamerProtocol(), InputHandler->GetFromStreamerProtocol() };
 		for (TSharedPtr<IPixelStreaming2DataProtocol> Protocol : Protocols)
 		{
@@ -1099,6 +1143,11 @@
 				{
 					TSharedPtr<FPlayerContext>& Participant = Participants->FindOrAdd(PlayerId);
 					Participant = MakeShared<FPlayerContext>();
+					if (!InputHandler)
+					{
+						UE_LOGFMT(LogPixelStreaming2RTC, Warning, "Cannot create SFU data track with an invalid input handler!");
+						return;
+					}
 					Participant->DataTrack = FEpicRtcMutliplexDataTrack::Create(SFUDataTrack, InputHandler->GetFromStreamerProtocol(), PlayerId);
 					OnDataTrackOpen(PlayerId);
 				}
@@ -1130,7 +1179,10 @@
 				{
 					TArray<uint8>							EmptyArray;
 					TFunction<void(FString, FMemoryReader)> MouseLeaveHandler = SharedInputHandler->FindMessageHandler("MouseLeave");
-					MouseLeaveHandler("", FMemoryReader(EmptyArray));
+					if (MouseLeaveHandler)
+					{
+						MouseLeaveHandler("", FMemoryReader(EmptyArray));
+					}
 				}
 			};
 
@@ -1760,6 +1812,11 @@
 		{
 			if (!Participant->DataTrack)
 			{
+				if (!InputHandler)
+				{
+					UE_LOGFMT(LogPixelStreaming2RTC, Warning, "Cannot create participant data track when input handler is invalid!");
+					return;
+				}
 				Participant->DataTrack = FEpicRtcDataTrack::Create(DataTrack, InputHandler->GetFromStreamerProtocol());
 			}
 			else
@@ -2033,7 +2090,12 @@
 		FString DataTrackId = ToString(DataTrack->GetId());
 		FString PlayerId;
 
-		uint8									 Type = Data[0];
+		uint8 Type = Data[0];
+		if (!InputHandler)
+		{
+			UE_LOGFMT(LogPixelStreaming2RTC, Warning, "Cannot handle data track message when input handler is invalid!");
+			return;
+		}
 		TSharedPtr<IPixelStreaming2DataProtocol> ToStreamerProtocol = InputHandler->GetToStreamerProtocol();
 		if (Type == ToStreamerProtocol->Find(EPixelStreaming2ToStreamerMessage::Multiplexed)->GetID())
 		{
@@ -2121,6 +2183,7 @@
 	TSharedPtr<IPixelStreaming2Streamer> FRTCStreamerFactory::CreateNewStreamer(const FString& StreamerId)
 	{
 		TSharedPtr<IPixelStreaming2Streamer> NewStreamer = MakeShared<FEpicRtcStreamer>(StreamerId, EpicRtcConference);
+		NewStreamer->SetInputHandler(FRTCInputHandler::Create());
 
 		// default to the scene viewport if we have a game engine
 		if (UGameEngine* GameEngine = Cast<UGameEngine>(GEngine))
diff --git a/Engine/Plugins/Media/PixelStreaming2/Source/PixelStreaming2RTC/Private/EpicRtcStreamer.h b/Engine/Plugins/Media/PixelStreaming2/Source/PixelStreaming2RTC/Private/EpicRtcStreamer.h
--- a/Engine/Plugins/Media/PixelStreaming2/Source/PixelStreaming2RTC/Private/EpicRtcStreamer.h
+++ b/Engine/Plugins/Media/PixelStreaming2/Source/PixelStreaming2RTC/Private/EpicRtcStreamer.h
@@ -82,6 +82,7 @@ namespace UE::PixelStreaming2
 		virtual TArray<FString> GetConnectedPlayers() override;
 
 		virtual TWeakPtr<IPixelStreaming2InputHandler> GetInputHandler() override { return InputHandler; }
+		virtual void SetInputHandler(TSharedPtr<IPixelStreaming2InputHandler> Handler) override;
 
 		virtual TWeakPtr<IPixelStreaming2AudioSink> GetPeerAudioSink(FString PlayerId) override;
 		virtual TWeakPtr<IPixelStreaming2AudioSink> GetUnlistenedAudioSink() override;
diff --git a/Engine/Plugins/Media/PixelStreaming2/Source/PixelStreaming2RTC/Private/RTCInputHandler.cpp b/Engine/Plugins/Media/PixelStreaming2/Source/PixelStreaming2RTC/Private/RTCInputHandler.cpp
--- a/Engine/Plugins/Media/PixelStreaming2/Source/PixelStreaming2RTC/Private/RTCInputHandler.cpp
+++ b/Engine/Plugins/Media/PixelStreaming2/Source/PixelStreaming2RTC/Private/RTCInputHandler.cpp
@@ -2,14 +2,14 @@
 
 #include "RTCInputHandler.h"
 
-#include "DefaultDataProtocol.h"
 #include "Engine/Engine.h"
 #include "Framework/Application/SlateApplication.h"
 #include "Framework/Application/SlateUser.h"
-#include "InputStructures.h"
+#include "PixelStreaming2InputStructures.h"
 #include "IPixelStreaming2HMDModule.h"
-#include "JavaScriptKeyCodes.inl"
 #include "Logging.h"
+#include "PixelStreaming2DefaultDataProtocol.h"
+#include "PixelStreaming2JavaScriptKeyCodes.inl"
 #include "PixelStreaming2PluginSettings.h"
 #include "PixelStreaming2Utils.h"
 #include "Widgets/Input/SEditableText.h"
@@ -545,11 +545,28 @@ namespace UE::PixelStreaming2
 	 */
 	void FRTCInputHandler::HandleOnCommand(FString SourceId, FMemoryReader Ar)
 	{
-		FString Res;
-		Res.GetCharArray().SetNumUninitialized(Ar.TotalSize() / 2 + 1);
-		Ar.Serialize(Res.GetCharArray().GetData(), Ar.TotalSize());
+		if (Ar.TotalSize() < 2)
+		{
+			UE_LOGFMT(LogPixelStreaming2RTC, Warning, "FRTCInputHandler::HandleOnTextboxEntry: Malformed message. Expected to contain at least 2 bytes representing the descriptor length in characters");
+			return;
+		}
+
+		uint16 StringLength = 0;
+		Ar << StringLength;
+
+		if (Ar.TotalSize() < static_cast<int64>(2 + StringLength * sizeof(TCHAR)))
+		{
+			UE_LOGFMT(LogPixelStreaming2RTC, Warning, "FRTCInputHandler::HandleOnTextboxEntry: Malformed message. Descriptor length of {0} characters exceeds remaining message size", StringLength);
+			return;
+		}
+
+		FString Descriptor;
+		Descriptor.GetCharArray().SetNumUninitialized(StringLength + 1 /* +1 for null terminator */);
+		// Read archive into the string's char array
+		Ar.Serialize(Descriptor.GetCharArray().GetData(), StringLength * sizeof(TCHAR));
+		// Null terminate
+		Descriptor.GetCharArray()[StringLength] = TEXT('\0');
 
-		FString Descriptor = Res.Mid(MessageHeaderOffset);
 		UE_LOGFMT(LogPixelStreaming2RTC, Verbose, "Command: {0}", Descriptor);
 
 		// Iterate each command handler and see if the command we got matches any of the bound command names.
@@ -580,13 +597,35 @@ namespace UE::PixelStreaming2
 	 */
 	void FRTCInputHandler::HandleOnTextboxEntry(FMemoryReader Ar)
 	{
-		FString Res;
-		Res.GetCharArray().SetNumUninitialized(Ar.TotalSize() / 2 + 1);
-		Ar.Serialize(Res.GetCharArray().GetData(), Ar.TotalSize());
-		FString Text = Res.Mid(1);
+		if (Ar.TotalSize() < 2)
+		{
+			UE_LOGFMT(LogPixelStreaming2RTC, Warning, "FRTCInputHandler::HandleOnTextboxEntry: Malformed message. Expected to contain at least 2 bytes representing the descriptor length in characters");
+			return;
+		}
+
+		uint16 StringLength = 0;
+		Ar << StringLength;
+
+		if (Ar.TotalSize() < static_cast<int64>(2 + StringLength * sizeof(TCHAR)))
+		{
+		UE_LOGFMT(LogPixelStreaming2RTC, Warning, "FRTCInputHandler::HandleOnTextboxEntry: Malformed message. Descriptor length of {0} characters exceeds remaining message size", StringLength);
+			return;
+		}
+
+		FString Text;
+		Text.GetCharArray().SetNumUninitialized(StringLength + 1 /* +1 for null terminator */);
+		// Read archive into the string's char array
+		Ar.Serialize(Text.GetCharArray().GetData(), StringLength * sizeof(TCHAR));
+		// Null terminate
+		Text.GetCharArray()[StringLength] = TEXT('\0');
 
 		FSlateApplication::Get().ForEachUser([this, Text](FSlateUser& User) {
 			TSharedPtr<SWidget> FocusedWidget = User.GetFocusedWidget();
+			if (!FocusedWidget)
+			{
+				UE_LOGFMT(LogPixelStreaming2RTC, Verbose, "FRTCInputHandler::HandleOnTextboxEntry: Focused widget is invalid");
+				return;
+			}
 
 			bool bIsEditableTextType = FocusedWidget->GetType() == TEXT("SEditableText");
 			bool bIsMultiLineEditableTextType = FocusedWidget->GetType() == TEXT("SMultiLineEditableText");
diff --git a/Engine/Plugins/Media/PixelStreaming2/Source/PixelStreaming2RTC/Private/Tests/MockPlayer.cpp b/Engine/Plugins/Media/PixelStreaming2/Source/PixelStreaming2RTC/Private/Tests/MockPlayer.cpp
--- a/Engine/Plugins/Media/PixelStreaming2/Source/PixelStreaming2RTC/Private/Tests/MockPlayer.cpp
+++ b/Engine/Plugins/Media/PixelStreaming2/Source/PixelStreaming2RTC/Private/Tests/MockPlayer.cpp
@@ -3,11 +3,11 @@
 #include "MockPlayer.h"
 
 #include "Async/Async.h"
-#include "DefaultDataProtocol.h"
 #include "EpicRtcVideoEncoderInitializer.h"
 #include "EpicRtcVideoDecoderInitializer.h"
 #include "EpicRtcWebsocketFactory.h"
 #include "Logging.h"
+#include "PixelStreaming2DefaultDataProtocol.h"
 #include "UtilsString.h"
 
 #include "epic_rtc/core/platform.h"
