[PixelStreaming2] Fix: Undo accidental change in FVideoProducerRenderTarget::GetName

#Synced-CL 48443255
#Synced-CL 48443255

diff --git a/Engine/Plugins/Media/PixelStreaming2/Source/PixelStreaming2/Private/VideoProducerRenderTarget.cpp b/Engine/Plugins/Media/PixelStreaming2/Source/PixelStreaming2/Private/VideoProducerRenderTarget.cpp
--- a/Engine/Plugins/Media/PixelStreaming2/Source/PixelStreaming2/Private/VideoProducerRenderTarget.cpp
+++ b/Engine/Plugins/Media/PixelStreaming2/Source/PixelStreaming2/Private/VideoProducerRenderTarget.cpp
@@ -50,7 +50,12 @@ namespace UE::PixelStreaming2
 
 	FString FVideoProducerRenderTarget::ToString()
 	{
-		return TEXT("CustomVideoProducer");
+		if (Target)
+		{
+			FString TargetName = Target->GetName();
+			return VideoProducerIdentifiers::FVideoProducerRenderTarget + FString::Printf(TEXT(" (%s)"), *TargetName);
+		}
+		return VideoProducerIdentifiers::FVideoProducerRenderTarget;
 	}
 
 } // namespace UE::PixelStreaming2
\ No newline at end of file
