[PixelCapture, PixelStreaming2] Fix: Multiple issues for RemoteSession on Android
- Add: InitializeDefaultStreamer setting for playback only use cases
- Fix: Set staging texture pixel format to RGBA8 for OpenGL

#Synced-CL 48443255

diff --git a/Engine/Plugins/Media/PixelCapture/Source/PixelCapture/Private/PixelCaptureCapturerI420ToRHI.cpp b/Engine/Plugins/Media/PixelCapture/Source/PixelCapture/Private/PixelCaptureCapturerI420ToRHI.cpp
--- a/Engine/Plugins/Media/PixelCapture/Source/PixelCapture/Private/PixelCaptureCapturerI420ToRHI.cpp
+++ b/Engine/Plugins/Media/PixelCapture/Source/PixelCapture/Private/PixelCaptureCapturerI420ToRHI.cpp
@@ -45,8 +45,11 @@ void FPixelCaptureCapturerI420ToRHI::Initialize(int32 InputWidth, int32 InputHei
 
 IPixelCaptureOutputFrame* FPixelCaptureCapturerI420ToRHI::CreateOutputBuffer(int32 InputWidth, int32 InputHeight)
 {
+	// OpenGL needs PF_R8G8B8A8
+	EPixelFormat PixelFormat = RHIGetInterfaceType() == ERHIInterfaceType::OpenGL ? PF_R8G8B8A8 : PF_B8G8R8A8;
+
 	FRHITextureCreateDesc TextureDesc =
-		FRHITextureCreateDesc::Create2D(TEXT("FPixelCaptureCapturerI420ToRHI Texture"), InputWidth, InputHeight, EPixelFormat::PF_B8G8R8A8)
+		FRHITextureCreateDesc::Create2D(TEXT("FPixelCaptureCapturerI420ToRHI Texture"), InputWidth, InputHeight, PixelFormat)
 			.SetClearValue(FClearValueBinding::None)
 			.SetFlags(ETextureCreateFlags::CPUWritable | ETextureCreateFlags::RenderTargetable)
 			.SetInitialState(ERHIAccess::Present)
@@ -85,6 +88,8 @@ void FPixelCaptureCapturerI420ToRHI::BeginProcess(const IPixelCaptureInputFrame&
 		TSharedPtr<FPixelCaptureBufferI420>		I420Buffer = SourceFrame.GetBuffer();
 		TSharedPtr<FPixelCaptureOutputFrameRHI> OutputTexture = StaticCastSharedPtr<FPixelCaptureOutputFrameRHI>(OutputBuffer);
 
+		uint32 FourCC = (RHIGetInterfaceType() == ERHIInterfaceType::OpenGL) ? libyuv::FOURCC_ABGR : libyuv::FOURCC_ARGB;
+
 		// convert the incoming i420 to argb
 		libyuv::ConvertFromI420(
 			I420Buffer->GetDataY(), I420Buffer->GetStrideY(),
@@ -92,7 +97,7 @@ void FPixelCaptureCapturerI420ToRHI::BeginProcess(const IPixelCaptureInputFrame&
 			I420Buffer->GetDataV(), I420Buffer->GetStrideUV(),
 			ARGBBuffer, 0,
 			I420Buffer->GetWidth(), I420Buffer->GetHeight(),
-			libyuv::FOURCC_ARGB);
+			FourCC);
 
 		TWeakPtr<FPixelCaptureCapturerI420ToRHI> WeakSelf = StaticCastSharedRef<FPixelCaptureCapturerI420ToRHI>(AsShared());
 		ENQUEUE_RENDER_COMMAND(I420toRHI)
diff --git a/Engine/Plugins/Media/PixelStreaming2/Source/PixelStreaming2/Private/PixelStreaming2Module.cpp b/Engine/Plugins/Media/PixelStreaming2/Source/PixelStreaming2/Private/PixelStreaming2Module.cpp
--- a/Engine/Plugins/Media/PixelStreaming2/Source/PixelStreaming2/Private/PixelStreaming2Module.cpp
+++ b/Engine/Plugins/Media/PixelStreaming2/Source/PixelStreaming2/Private/PixelStreaming2Module.cpp
@@ -45,12 +45,20 @@ namespace UE::PixelStreaming2
 		}
 
 		const ERHIInterfaceType RHIType = GDynamicRHI ? RHIGetInterfaceType() : ERHIInterfaceType::Hidden;
-		// only D3D11/D3D12/Vulkan is supported
-		if (!(RHIType == ERHIInterfaceType::D3D11 || RHIType == ERHIInterfaceType::D3D12 || RHIType == ERHIInterfaceType::Vulkan || RHIType == ERHIInterfaceType::Metal))
-		{
-#if !WITH_DEV_AUTOMATION_TESTS
-			UE_LOG(LogPixelStreaming2, Warning, TEXT("Only D3D11/D3D12/Vulkan/Metal Dynamic RHI is supported. Detected %s"), GDynamicRHI != nullptr ? GDynamicRHI->GetName() : TEXT("[null]"));
+		// Validate RHI against what PS2 supports
+		if (!(
+			RHIType == ERHIInterfaceType::D3D11 
+			|| RHIType == ERHIInterfaceType::D3D12 
+			|| RHIType == ERHIInterfaceType::Vulkan 
+			|| RHIType == ERHIInterfaceType::Metal
+#if PLATFORM_ANDROID
+			|| RHIType == ERHIInterfaceType::OpenGL
 #endif
+			))
+		{
+	#if !WITH_DEV_AUTOMATION_TESTS
+			UE_LOG(LogPixelStreaming2, Warning, TEXT("Only D3D11/D3D12/Vulkan/Metal/OpenGL(Android only) Dynamic RHI is supported. Detected %s"), GDynamicRHI != nullptr ? GDynamicRHI->GetName() : TEXT("[null]"));
+	#endif
 			return;
 		}
 
@@ -93,7 +101,7 @@ namespace UE::PixelStreaming2
 
 			// This is called after the ready event is broadcast so that other modules have a chance to initialize
 			// themselves before streamer creation
-			if (!GIsEditor)
+			if (!GIsEditor && UPixelStreaming2PluginSettings::CVarInitializeDefaultStreamer.GetValueOnAnyThread())
 			{
 				// We don't want to create the default streamer when using the editor
 				InitDefaultStreamer();
@@ -180,7 +188,7 @@ namespace UE::PixelStreaming2
 		});
 	}
 
-	TSharedPtr<IPixelStreaming2Streamer> FPixelStreaming2Module::CreateStreamer(const FString& StreamerId, const FString& Type)
+	TSharedPtr<IPixelStreaming2Streamer> FPixelStreaming2Module::CreateStreamer(const FString& StreamerId, TOptional<FString> Type)
 	{
 		TSharedPtr<IPixelStreaming2Streamer> ExistingStreamer = FindStreamer(StreamerId);
 		if (ExistingStreamer)
@@ -188,7 +196,7 @@ namespace UE::PixelStreaming2
 			return ExistingStreamer;
 		}
 
-		IPixelStreaming2StreamerFactory*	 Factory = IPixelStreaming2StreamerFactory::Get(Type);
+		IPixelStreaming2StreamerFactory*	 Factory = IPixelStreaming2StreamerFactory::Get(Type.Get(UPixelStreaming2PluginSettings::CVarDefaultStreamerType.GetValueOnAnyThread()));
 		TSharedPtr<IPixelStreaming2Streamer> NewStreamer = Factory->CreateNewStreamer(StreamerId);
 		NewStreamer->Initialize();
 		// Any time we create a new streamer, populate it's signalling server URL with whatever is in the ini, console or command line
diff --git a/Engine/Plugins/Media/PixelStreaming2/Source/PixelStreaming2/Private/PixelStreaming2Module.h b/Engine/Plugins/Media/PixelStreaming2/Source/PixelStreaming2/Private/PixelStreaming2Module.h
--- a/Engine/Plugins/Media/PixelStreaming2/Source/PixelStreaming2/Private/PixelStreaming2Module.h
+++ b/Engine/Plugins/Media/PixelStreaming2/Source/PixelStreaming2/Private/PixelStreaming2Module.h
@@ -23,7 +23,7 @@ namespace UE::PixelStreaming2
 		virtual bool									  IsReady() override;
 		virtual bool									  StartStreaming() override;
 		virtual void									  StopStreaming() override;
-		virtual TSharedPtr<IPixelStreaming2Streamer>	  CreateStreamer(const FString& StreamerId, const FString& Type = TEXT("DefaultRtc")) override;
+		virtual TSharedPtr<IPixelStreaming2Streamer>	  CreateStreamer(const FString& StreamerId, TOptional<FString> Type = TOptional<FString>()) override;
 		virtual TSharedPtr<IPixelStreaming2VideoProducer> CreateVideoProducer() override;
 		virtual TSharedPtr<IPixelStreaming2AudioProducer> CreateAudioProducer() override;
 		virtual TArray<FString>							  GetStreamerIds() override;
diff --git a/Engine/Plugins/Media/PixelStreaming2/Source/PixelStreaming2/Public/IPixelStreaming2Module.h b/Engine/Plugins/Media/PixelStreaming2/Source/PixelStreaming2/Public/IPixelStreaming2Module.h
--- a/Engine/Plugins/Media/PixelStreaming2/Source/PixelStreaming2/Public/IPixelStreaming2Module.h
+++ b/Engine/Plugins/Media/PixelStreaming2/Source/PixelStreaming2/Public/IPixelStreaming2Module.h
@@ -73,7 +73,7 @@ public:
 	 * @param StreamerId - The ID of the Streamer to be created.
 	 * @return SharedPtr to the streamer.
 	 */
-	virtual TSharedPtr<IPixelStreaming2Streamer> CreateStreamer(const FString& StreamerId, const FString& Type = TEXT("DefaultRtc")) = 0;
+	virtual TSharedPtr<IPixelStreaming2Streamer> CreateStreamer(const FString& StreamerId, TOptional<FString> Type = TOptional<FString>()) = 0;
 
 	/**
 	 * @brief Creates a new video producer.
diff --git a/Engine/Plugins/Media/PixelStreaming2/Source/PixelStreaming2RTC/Private/PixelStreaming2RTCModule.cpp b/Engine/Plugins/Media/PixelStreaming2/Source/PixelStreaming2RTC/Private/PixelStreaming2RTCModule.cpp
--- a/Engine/Plugins/Media/PixelStreaming2/Source/PixelStreaming2RTC/Private/PixelStreaming2RTCModule.cpp
+++ b/Engine/Plugins/Media/PixelStreaming2/Source/PixelStreaming2RTC/Private/PixelStreaming2RTCModule.cpp
@@ -91,11 +91,19 @@ namespace UE::PixelStreaming2
 		}
 
 		const ERHIInterfaceType RHIType = GDynamicRHI ? RHIGetInterfaceType() : ERHIInterfaceType::Hidden;
-		// only D3D11/D3D12/Vulkan is supported
-		if (!(RHIType == ERHIInterfaceType::D3D11 || RHIType == ERHIInterfaceType::D3D12 || RHIType == ERHIInterfaceType::Vulkan || RHIType == ERHIInterfaceType::Metal))
+		// Validate RHI against what PS2 supports
+		if (!(
+			RHIType == ERHIInterfaceType::D3D11 
+			|| RHIType == ERHIInterfaceType::D3D12 
+			|| RHIType == ERHIInterfaceType::Vulkan 
+			|| RHIType == ERHIInterfaceType::Metal
+#if PLATFORM_ANDROID
+			|| RHIType == ERHIInterfaceType::OpenGL
+#endif
+			))
 		{
 	#if !WITH_DEV_AUTOMATION_TESTS
-			UE_LOG(LogPixelStreaming2RTC, Warning, TEXT("Only D3D11/D3D12/Vulkan/Metal Dynamic RHI is supported. Detected %s"), GDynamicRHI != nullptr ? GDynamicRHI->GetName() : TEXT("[null]"));
+			UE_LOG(LogPixelStreaming2RTC, Warning, TEXT("Only D3D11/D3D12/Vulkan/Metal/OpenGL(Android only) Dynamic RHI is supported. Detected %s"), GDynamicRHI != nullptr ? GDynamicRHI->GetName() : TEXT("[null]"));
 	#endif
 			return;
 		}
diff --git a/Engine/Plugins/Media/PixelStreaming2/Source/PixelStreaming2Settings/Internal/PixelStreaming2PluginSettings.h b/Engine/Plugins/Media/PixelStreaming2/Source/PixelStreaming2Settings/Internal/PixelStreaming2PluginSettings.h
--- a/Engine/Plugins/Media/PixelStreaming2/Source/PixelStreaming2Settings/Internal/PixelStreaming2PluginSettings.h
+++ b/Engine/Plugins/Media/PixelStreaming2/Source/PixelStreaming2Settings/Internal/PixelStreaming2PluginSettings.h
@@ -165,10 +165,18 @@ public:
 		))
 	bool UseMediaCapture = true;
 
+	static UE_API TAutoConsoleVariable<bool> CVarInitializeDefaultStreamer;
+	UPROPERTY(config, EditAnywhere, Category = "PixelStreaming", meta = (
+		DisplayName = "Initialize Default Streamer",
+		ToolTip = "If the default Pixel Streaming 2 streamer should be initialized. You may wish to set this to false if you only plan on receiving streams. Default: true"
+		))
+	bool InitializeDefaultStreamer = true;
+
 	static UE_API TAutoConsoleVariable<FString> CVarDefaultStreamerID;
 	UPROPERTY(config, EditAnywhere, Category = "PixelStreaming", meta = (
 		DisplayName = "Default Streamer ID",
-		ToolTip = "Default Streamer ID to be used when not specified elsewhere."
+		ToolTip = "Default Streamer ID to be used when not specified elsewhere.",
+		EditCondition="InitializeDefaultStreamer"
 		))
 	FString DefaultStreamerID = TEXT("DefaultStreamer");
 
@@ -176,14 +184,16 @@ public:
 	UPROPERTY(config, EditAnywhere, Category = "PixelStreaming", meta = (
 		DisplayName = "Default Streamer Type",
 		ToolTip = "Default Streamer Type to be used when not specified elsewhere. This value should match a type registered by an IPixelStreaming2StreamerFactory",
-		GetOptions = "GetDefaultStreamerTypeOptions"
+		GetOptions = "GetDefaultStreamerTypeOptions",
+		EditCondition="InitializeDefaultStreamer"
 		))
 	FString DefaultStreamerType = TEXT("DefaultRtc");
 
 	static UE_API TAutoConsoleVariable<bool> CVarAutoStartStream;
 	UPROPERTY(config, EditAnywhere, Category = "PixelStreaming", meta = (
 		DisplayName = "Automatically Start Streaming",
-		ToolTip = "Configure the PixelStreaming2 plugin to automatically start streaming once loaded (if not in editor). You may wish to set this value to false and manually call StartStreaming at a later point from your c++ code. Default: true"
+		ToolTip = "Configure the PixelStreaming2 plugin to automatically start streaming once loaded (if not in editor). You may wish to set this value to false and manually call StartStreaming at a later point from your c++ code. Default: true",
+		EditCondition="InitializeDefaultStreamer"
 		))
 	bool AutoStartStream = true;
 
diff --git a/Engine/Plugins/Media/PixelStreaming2/Source/PixelStreaming2Settings/Private/PixelStreaming2PluginSettings.cpp b/Engine/Plugins/Media/PixelStreaming2/Source/PixelStreaming2Settings/Private/PixelStreaming2PluginSettings.cpp
--- a/Engine/Plugins/Media/PixelStreaming2/Source/PixelStreaming2Settings/Private/PixelStreaming2PluginSettings.cpp
+++ b/Engine/Plugins/Media/PixelStreaming2/Source/PixelStreaming2Settings/Private/PixelStreaming2PluginSettings.cpp
@@ -173,6 +173,7 @@ static const TSet<TPair<FString, FString>> GetCmdArg = {
 	{ "PixelStreaming2.SignalingMaxReconnectAttempts", "SignalingMaxReconnectAttempts" },
 	{ "PixelStreaming2.SignalingKeepAliveInterval", "SignalingKeepAliveInterval" },
 	{ "PixelStreaming2.UseMediaCapture", "UseMediaCapture" },
+	{ "PixelStreaming2.InitializeDefaultStreamer", "InitializeDefaultStreamer" },
 	{ "PixelStreaming2.ID", "DefaultStreamerID" },
 	{ "PixelStreaming2.DefaultStreamerType", "DefaultStreamerType" },
 	{ "PixelStreaming2.AutoStartStream", "AutoStartStream" },
@@ -337,6 +338,12 @@ TAutoConsoleVariable<bool> UPixelStreaming2PluginSettings::CVarUseMediaCapture(
 		}),
 	ECVF_Default);
 
+TAutoConsoleVariable<bool> UPixelStreaming2PluginSettings::CVarInitializeDefaultStreamer(
+	TEXT("PixelStreaming2.InitializeDefaultStreamer"),
+	true,
+	TEXT("If the default Pixel Streaming 2 streamer should be initialized. You may wish to set this to false if you only plan on receiving streams. Default: true"),
+	ECVF_Default);
+
 TAutoConsoleVariable<FString> UPixelStreaming2PluginSettings::CVarDefaultStreamerID(
 	TEXT("PixelStreaming2.ID"),
 	TEXT("DefaultStreamer"),
