[PixelStreaming2] Add WebSocket origin filtering (CORS) to signalling server

  Adds configurable origin-based connection filtering to the PixelStreaming2
  signalling server using the existing IWebSocketServer::SetFilterConnectionCallback
  API. Connections are validated against an allowed origins list before the
  WebSocket handshake completes.

  - Expose SetFilterConnectionCallback on FWebSocketServerWrapper
  - Add FilterConnection() to FSignallingServer with wildcard pattern support via FString::MatchesWildcard
  - Filter applied to both StreamersWS and PlayersWS servers
  - New CVar: PixelStreaming2.InternalSignallingServer.AllowedOrigins (comma-separated, supports wildcards)
  - New UPROPERTY in UPixelStreaming2PluginSettings under SignallingServer category
  - Configurable via Project Settings, INI, command line, or console
  - Default: http://127.0.0.1,https://127.0.0.1
  - Set to * to allow all origins, or leave empty to disable filtering

#Synced-CL 52245801

diff --git a/Engine/Plugins/Media/PixelStreaming2/Source/PixelStreaming2Editor/Private/PixelStreaming2EditorModule.cpp b/Engine/Plugins/Media/PixelStreaming2/Source/PixelStreaming2Editor/Private/PixelStreaming2EditorModule.cpp
--- a/Engine/Plugins/Media/PixelStreaming2/Source/PixelStreaming2Editor/Private/PixelStreaming2EditorModule.cpp
+++ b/Engine/Plugins/Media/PixelStreaming2/Source/PixelStreaming2Editor/Private/PixelStreaming2EditorModule.cpp
@@ -359,6 +359,12 @@ namespace UE::EditorPixelStreaming2
 		FString CertificateArgs = FString::Printf(TEXT("--CertificatePath=%s --PrivateKeyPath=%s"), *SSLCertificatePath, *SSLPrivateKeyPath);
 		LaunchArgs.ProcessArgs = FString::Printf(TEXT("--HttpPort=%d --StreamerPort=%d --ServeHttps=%s %s"), ViewerPort, StreamerPort, bServeHttps ? TEXT("true") : TEXT("false"), bServeHttps ? *CertificateArgs : TEXT(""));
 
+		FString AllowedOrigins = UPixelStreaming2PluginSettings::CVarAllowedOrigins.GetValueOnAnyThread();
+		if (!AllowedOrigins.IsEmpty())
+		{
+			LaunchArgs.ProcessArgs += FString::Printf(TEXT(" --AllowedOrigins=%s"), *AllowedOrigins);
+		}
+
 		SignallingServer->Launch(LaunchArgs);
 	}
 
diff --git a/Engine/Plugins/Media/PixelStreaming2/Source/PixelStreaming2Editor/Private/PixelStreaming2Toolbar.cpp b/Engine/Plugins/Media/PixelStreaming2/Source/PixelStreaming2Editor/Private/PixelStreaming2Toolbar.cpp
--- a/Engine/Plugins/Media/PixelStreaming2/Source/PixelStreaming2Editor/Private/PixelStreaming2Toolbar.cpp
+++ b/Engine/Plugins/Media/PixelStreaming2/Source/PixelStreaming2Editor/Private/PixelStreaming2Toolbar.cpp
@@ -465,6 +465,47 @@ namespace UE::EditorPixelStreaming2
 			];
 		MenuBuilder.AddWidget(ViewerPortInputBlock, FText(), true);
 
+		TSharedRef<SWidget> AllowedOriginsInputBlock = SNew(SHorizontalBox)
+			.Visibility_Lambda([]() {
+				return !GetUseRemoteSignallingServer() ? EVisibility::Visible : EVisibility::Collapsed;
+			})
+			+ SHorizontalBox::Slot()
+			.AutoWidth()
+			.VAlign(VAlign_Center)
+			.Padding(FMargin(36.0f, 3.0f, 8.0f, 3.0f))
+			[
+				SNew(STextBlock)
+					.Text(FText::FromString(TEXT("Allowed Origins: ")))
+					.ToolTipText(LOCTEXT("PixelStreaming2AllowedOriginsTooltip", "Comma-separated list of allowed origins for WebSocket connection filtering. Supports wildcards (*). Use '*' to allow all origins."))
+					.ColorAndOpacity(FSlateColor(FLinearColor(1.0f, 1.0f, 1.0f)))
+			]
+			+ SHorizontalBox::Slot()
+			.FillWidth(1.0f)
+			.MaxWidth(200.0f)
+			[
+				SNew(SEditableTextBox)
+					.Text_Lambda([]() {
+						return FText::FromString(UPixelStreaming2PluginSettings::CVarAllowedOrigins.GetValueOnAnyThread());
+					})
+					.OnTextChanged_Lambda([](const FText& InText) {
+						UPixelStreaming2PluginSettings::CVarAllowedOrigins.AsVariable()->SetWithCurrentPriority(*InText.ToString());
+					})
+					.OnTextCommitted_Lambda([](const FText& InText, ETextCommit::Type InCommitType) {
+						UPixelStreaming2PluginSettings::CVarAllowedOrigins.AsVariable()->SetWithCurrentPriority(*InText.ToString());
+						// Save to ini so it persists across sessions
+						if (UPixelStreaming2PluginSettings* Settings = GetMutableDefault<UPixelStreaming2PluginSettings>())
+						{
+							Settings->AllowedOrigins = InText.ToString();
+							Settings->SaveConfig();
+						}
+					})
+					.IsEnabled_Lambda([]() {
+						TSharedPtr<UE::PixelStreaming2Servers::IServer> SignallingServer = IPixelStreaming2EditorModule::Get().GetSignallingServer();
+						return !(SignallingServer.IsValid() && SignallingServer->HasLaunched());
+					})
+			];
+		MenuBuilder.AddWidget(AllowedOriginsInputBlock, FText(), true);
+
 		MenuBuilder.AddMenuEntry(
     		FPixelStreaming2Commands::Get().StartSignalling,       // Command
     		NAME_None,                                              // InExtensionHook
diff --git a/Engine/Plugins/Media/PixelStreaming2/Source/PixelStreaming2Servers/Private/SignallingServer.cpp b/Engine/Plugins/Media/PixelStreaming2/Source/PixelStreaming2Servers/Private/SignallingServer.cpp
--- a/Engine/Plugins/Media/PixelStreaming2/Source/PixelStreaming2Servers/Private/SignallingServer.cpp
+++ b/Engine/Plugins/Media/PixelStreaming2/Source/PixelStreaming2Servers/Private/SignallingServer.cpp
@@ -61,16 +61,100 @@ namespace UE::PixelStreaming2Servers
 		}
 	}
 
+	// Returns true if the origin string contains a port (e.g. "http://localhost:8080", "http://[::1]:8080")
+	// by checking for a colon after the host portion of the URL.
+	// Handles IPv6 addresses in brackets (e.g. "http://[::1]") where colons inside brackets are not ports.
+	static bool OriginHasPort(const FString& InOrigin)
+	{
+		int32 SchemeEnd;
+		if (InOrigin.FindChar(TEXT('/'), SchemeEnd) && InOrigin.Mid(SchemeEnd).StartsWith(TEXT("//")))
+		{
+			FString AfterScheme = InOrigin.Mid(SchemeEnd + 3);
+
+			// IPv6: skip past the bracketed address before looking for a port colon
+			if (AfterScheme.StartsWith(TEXT("[")))
+			{
+				int32 BracketEnd;
+				if (AfterScheme.FindChar(TEXT(']'), BracketEnd))
+				{
+					return AfterScheme.Mid(BracketEnd + 1).StartsWith(TEXT(":"));
+				}
+				return false;
+			}
+
+			// IPv4 / hostname: a colon after the scheme means a port
+			return AfterScheme.Contains(TEXT(":"));
+		}
+		return false;
+	}
+
+	EWebsocketConnectionFilterResult FSignallingServer::FilterConnection(FString Origin, FString ClientIP)
+	{
+		// If no allowed origins configured, accept all connections
+		if (AllowedOrigins.Num() == 0)
+		{
+			return EWebsocketConnectionFilterResult::ConnectionAccepted;
+		}
+
+		// Wildcard allows all origins we normalize to by adding":*" even if the user passes *
+		if (AllowedOrigins.Contains(TEXT("*:*")))
+		{
+			return EWebsocketConnectionFilterResult::ConnectionAccepted;
+		}
+
+		// Normalize: if the origin has no port, append ":*" so it can match patterns with ports
+		FString NormalizedOrigin = Origin;
+		if (!OriginHasPort(NormalizedOrigin))
+		{
+			NormalizedOrigin += TEXT(":*");
+		}
+
+		// Check if the origin matches any allowed pattern (supports wildcards e.g. http://192.168.0.*:8080)
+		for (const FString& AllowedOrigin : AllowedOrigins)
+		{
+			if (NormalizedOrigin.MatchesWildcard(AllowedOrigin, ESearchCase::IgnoreCase))
+			{
+				return EWebsocketConnectionFilterResult::ConnectionAccepted;
+			}
+		}
+
+		UE_LOGFMT(LogPixelStreaming2Servers, Warning, "Rejected WebSocket connection from origin '{0}' (client IP: {1}). Allowed origins: [{2}]", Origin, ClientIP, FString::Join(AllowedOrigins, TEXT(", ")));
+		return EWebsocketConnectionFilterResult::ConnectionRefused;
+	}
+
 	bool FSignallingServer::LaunchImpl(FLaunchArgs& InLaunchArgs, TMap<EEndpoint, FURL>& OutEndpoints)
 	{
 		Utils::PopulateCirrusEndPoints(InLaunchArgs, OutEndpoints);
 		FURL PlayersURL = OutEndpoints[EEndpoint::Signalling_Players];
 		FURL StreamerURL = OutEndpoints[EEndpoint::Signalling_Streamer];
 
+		// Parse allowed origins for CORS filtering
+		FString AllowedOriginsArg = Utils::QueryOrSetProcessArgs(InLaunchArgs, TEXT("--AllowedOrigins="), TEXT(""));
+		if (!AllowedOriginsArg.IsEmpty())
+		{
+			AllowedOriginsArg.ParseIntoArray(AllowedOrigins, TEXT(","), true);
+			for (FString& AllowedOrigin : AllowedOrigins)
+			{
+				AllowedOrigin.TrimStartAndEndInline();
+				// Normalize: if no port specified, append ":*" so it matches origins with any port
+				if (!OriginHasPort(AllowedOrigin))
+				{
+					AllowedOrigin += TEXT(":*");
+				}
+			}
+			UE_LOGFMT(LogPixelStreaming2Servers, Log, "CORS filtering enabled. Allowed origins: {0}", AllowedOriginsArg);
+		}
+
 		/*
 		 * --------------- Streamers websocket server ---------------
 		 */
 		StreamersWS = MakeUnique<FWebSocketServerWrapper>();
+		if (AllowedOrigins.Num() > 0)
+		{
+			FWebSocketFilterConnectionCallback FilterCallback;
+			FilterCallback.BindRaw(this, &FSignallingServer::FilterConnection);
+			StreamersWS->SetFilterConnectionCallback(MoveTemp(FilterCallback));
+		}
 		StreamersWS->OnMessage.AddRaw(this, &FSignallingServer::OnStreamerMessage);
 		StreamersWS->OnOpenConnection.AddRaw(this, &FSignallingServer::OnStreamerConnected);
 		StreamersWS->OnClosedConnection.AddRaw(this, &FSignallingServer::OnStreamerDisconnected);
@@ -104,6 +188,12 @@ namespace UE::PixelStreaming2Servers
 		 * --------------- Players websocket server ---------------
 		 */
 		PlayersWS = MakeUnique<FWebSocketServerWrapper>();
+		if (AllowedOrigins.Num() > 0)
+		{
+			FWebSocketFilterConnectionCallback FilterCallback;
+			FilterCallback.BindRaw(this, &FSignallingServer::FilterConnection);
+			PlayersWS->SetFilterConnectionCallback(MoveTemp(FilterCallback));
+		}
 		PlayersWS->EnableWebServer(GenerateDirectoriesToServe(), bServeHttps, Certificates);
 		PlayersWS->OnMessage.AddRaw(this, &FSignallingServer::OnPlayerMessage);
 		PlayersWS->OnOpenConnection.AddRaw(this, &FSignallingServer::OnPlayerConnected);
diff --git a/Engine/Plugins/Media/PixelStreaming2/Source/PixelStreaming2Servers/Private/SignallingServer.h b/Engine/Plugins/Media/PixelStreaming2/Source/PixelStreaming2Servers/Private/SignallingServer.h
--- a/Engine/Plugins/Media/PixelStreaming2/Source/PixelStreaming2Servers/Private/SignallingServer.h
+++ b/Engine/Plugins/Media/PixelStreaming2/Source/PixelStreaming2Servers/Private/SignallingServer.h
@@ -57,12 +57,15 @@ namespace UE::PixelStreaming2Servers
 		void OnPlayerUnsubscribeMessage(uint16 ConnectionId, TSharedPtr<FJsonObject> JSONObj);
 		void OnPlayerStatsMessage(uint16 ConnectionId, TSharedPtr<FJsonObject> JSONObj);
 		
+		EWebsocketConnectionFilterResult FilterConnection(FString Origin, FString ClientIP);
+
 	protected:
 		TUniquePtr<FWebSocketProbe> Probe;
 		TUniquePtr<FWebSocketServerWrapper> StreamersWS;
 		TUniquePtr<FWebSocketServerWrapper> PlayersWS;
 
 		TMap<uint16, uint16> PlayerSubscriptions;
+		TArray<FString> AllowedOrigins;
 
 		DECLARE_DELEGATE_TwoParams(FMessageHandler, uint16, TSharedPtr<FJsonObject>);
 		TMap<FString, FMessageHandler> StreamerMessageHandlers;
diff --git a/Engine/Plugins/Media/PixelStreaming2/Source/PixelStreaming2Servers/Private/WebSocketServerWrapper.cpp b/Engine/Plugins/Media/PixelStreaming2/Source/PixelStreaming2Servers/Private/WebSocketServerWrapper.cpp
--- a/Engine/Plugins/Media/PixelStreaming2/Source/PixelStreaming2Servers/Private/WebSocketServerWrapper.cpp
+++ b/Engine/Plugins/Media/PixelStreaming2/Source/PixelStreaming2Servers/Private/WebSocketServerWrapper.cpp
@@ -82,6 +82,15 @@ namespace UE::PixelStreaming2Servers
 		Certificates = InCertificates;
 	}
 
+	void FWebSocketServerWrapper::SetFilterConnectionCallback(FWebSocketFilterConnectionCallback InFilterConnectionCallback)
+	{
+		FilterConnectionCallback = MoveTemp(InFilterConnectionCallback);
+		if (WSServer)
+		{
+			WSServer->SetFilterConnectionCallback(FilterConnectionCallback);
+		}
+	}
+
 	bool FWebSocketServerWrapper::Launch(uint16 Port)
 	{
 		SetTickableTickType(ETickableTickType::Always);
@@ -93,6 +102,11 @@ namespace UE::PixelStreaming2Servers
 			WSServer->EnableHTTPServer(DirectoriesToServe, bServeHttps, Certificates);
 		}
 
+		if (FilterConnectionCallback.IsBound())
+		{
+			WSServer->SetFilterConnectionCallback(FilterConnectionCallback);
+		}
+
 		OnClientConnectedCallback.BindRaw(this, &FWebSocketServerWrapper::OnConnectionOpened);
 		bLaunched = WSServer->Init(Port, OnClientConnectedCallback);
 		if (!bLaunched)
diff --git a/Engine/Plugins/Media/PixelStreaming2/Source/PixelStreaming2Servers/Private/WebSocketServerWrapper.h b/Engine/Plugins/Media/PixelStreaming2/Source/PixelStreaming2Servers/Private/WebSocketServerWrapper.h
--- a/Engine/Plugins/Media/PixelStreaming2/Source/PixelStreaming2Servers/Private/WebSocketServerWrapper.h
+++ b/Engine/Plugins/Media/PixelStreaming2/Source/PixelStreaming2Servers/Private/WebSocketServerWrapper.h
@@ -37,6 +37,7 @@ namespace UE::PixelStreaming2Servers
 		FWebSocketServerWrapper();
 		virtual ~FWebSocketServerWrapper();
 		void EnableWebServer(TArray<FWebSocketHttpMount> InDirectoriesToServe, bool bServeHttps, const FWebSocketServerCertificates& Certificates);
+		void SetFilterConnectionCallback(FWebSocketFilterConnectionCallback InFilterConnectionCallback);
 		bool Launch(uint16 Port);
 		void Stop();
 		bool IsLaunched() const;
@@ -92,6 +93,7 @@ namespace UE::PixelStreaming2Servers
 		bool										   bServeHttps = false;
 		FWebSocketServerCertificates				   Certificates;
 		TArray<FWebSocketHttpMount>					   DirectoriesToServe;
+		FWebSocketFilterConnectionCallback			   FilterConnectionCallback;
 	};
 
 } // namespace UE::PixelStreaming2Servers
diff --git a/Engine/Plugins/Media/PixelStreaming2/Source/PixelStreaming2Settings/Internal/PixelStreaming2PluginSettings.h b/Engine/Plugins/Media/PixelStreaming2/Source/PixelStreaming2Settings/Internal/PixelStreaming2PluginSettings.h
--- a/Engine/Plugins/Media/PixelStreaming2/Source/PixelStreaming2Settings/Internal/PixelStreaming2PluginSettings.h
+++ b/Engine/Plugins/Media/PixelStreaming2/Source/PixelStreaming2Settings/Internal/PixelStreaming2PluginSettings.h
@@ -611,6 +611,15 @@ public:
 
 	// End Input CVars
 
+	// Begin Signalling Server CVars
+	static UE_API TAutoConsoleVariable<FString> CVarAllowedOrigins;
+	UPROPERTY(config, EditAnywhere, Category = "SignallingServer", meta = (
+		DisplayName = "Allowed Origins",
+		ToolTip = "Comma-separated list of allowed origins for WebSocket connection filtering. Supports wildcards (*). Use '*' to allow all origins. Leave empty to disable filtering. Examples: \"*\"  or \"https://example.com,http://192.168.0.*:8080\". Default: \"http://127.0.0.1,https://127.0.0.1\""
+		))
+	FString AllowedOrigins = TEXT("http://127.0.0.1,https://127.0.0.1");
+	// End Signalling Server CVars
+
 	// clang-format on
 
 	// Begin UDeveloperSettings Interface
diff --git a/Engine/Plugins/Media/PixelStreaming2/Source/PixelStreaming2Settings/Private/PixelStreaming2PluginSettings.cpp b/Engine/Plugins/Media/PixelStreaming2/Source/PixelStreaming2Settings/Private/PixelStreaming2PluginSettings.cpp
--- a/Engine/Plugins/Media/PixelStreaming2/Source/PixelStreaming2Settings/Private/PixelStreaming2PluginSettings.cpp
+++ b/Engine/Plugins/Media/PixelStreaming2/Source/PixelStreaming2Settings/Private/PixelStreaming2PluginSettings.cpp
@@ -222,7 +222,8 @@ static const TSet<TPair<FString, FString>> GetCmdArg = {
 	{ "PixelStreaming2.HMD.ProjectionOffsetY", "HMDProjectionOffsetY" },
 	{ "PixelStreaming2.AllowPixelStreamingCommands", "InputAllowConsoleCommands" },
 	{ "PixelStreaming2.KeyFilter", "InputKeyFilter" },
-	{ "PixelStreaming2.WebRTC.CodecPreferences", "WebRTCCodecPreferences" }
+	{ "PixelStreaming2.WebRTC.CodecPreferences", "WebRTCCodecPreferences" },
+	{ "PixelStreaming2.InternalSignallingServer.AllowedOrigins", "AllowedOrigins" }
 };
 
 static const TSet<TPair<FString, FString>> GetMappedCmdArg = {
@@ -725,6 +726,14 @@ TAutoConsoleVariable<FString> UPixelStreaming2PluginSettings::CVarInputKeyFilter
 	ECVF_Default);
 // End Input CVars
 
+// Begin Signalling Server CVars
+TAutoConsoleVariable<FString> UPixelStreaming2PluginSettings::CVarAllowedOrigins(
+	TEXT("PixelStreaming2.InternalSignallingServer.AllowedOrigins"),
+	TEXT("http://127.0.0.1,https://127.0.0.1"),
+	TEXT("Comma-separated list of allowed origins for WebSocket connection filtering. Supports wildcards (*). Use '*' to allow all origins. Leave empty to disable filtering."),
+	ECVF_Default);
+// End Signalling Server CVars
+
 TArray<EVideoCodec> UPixelStreaming2PluginSettings::GetCodecPreferences()
 {
 	TArray<EVideoCodec> OutCodecPreferences;
