Fixup test flags so tests can be enabled in test build configuration.

diff --git a/Engine/Plugins/Media/PixelStreaming2/Source/PixelStreaming2RTC/Private/Tests/MockPlayer.cpp b/Engine/Plugins/Media/PixelStreaming2/Source/PixelStreaming2RTC/Private/Tests/MockPlayer.cpp
--- a/Engine/Plugins/Media/PixelStreaming2/Source/PixelStreaming2RTC/Private/Tests/MockPlayer.cpp
+++ b/Engine/Plugins/Media/PixelStreaming2/Source/PixelStreaming2RTC/Private/Tests/MockPlayer.cpp
@@ -12,7 +12,7 @@
 
 #include "epic_rtc/core/platform.h"
 
-#if WITH_DEV_AUTOMATION_TESTS
+#if WITH_TESTS
 
 namespace UE::PixelStreaming2
 {
@@ -440,4 +440,4 @@ namespace UE::PixelStreaming2
 	}
 } // namespace UE::PixelStreaming2
 
-#endif // WITH_DEV_AUTOMATION_TESTS
+#endif // WITH_TESTS
diff --git a/Engine/Plugins/Media/PixelStreaming2/Source/PixelStreaming2RTC/Private/Tests/MockPlayer.h b/Engine/Plugins/Media/PixelStreaming2/Source/PixelStreaming2RTC/Private/Tests/MockPlayer.h
--- a/Engine/Plugins/Media/PixelStreaming2/Source/PixelStreaming2RTC/Private/Tests/MockPlayer.h
+++ b/Engine/Plugins/Media/PixelStreaming2/Source/PixelStreaming2RTC/Private/Tests/MockPlayer.h
@@ -17,7 +17,7 @@
 #include "EpicRtcRoomObserver.h"
 #include "EpicRtcSessionObserver.h"
 
-#if WITH_DEV_AUTOMATION_TESTS
+#if WITH_TESTS
 
 class EpicRtcPlatformInterface;
 
@@ -184,4 +184,4 @@ namespace UE::PixelStreaming2
 	};
 } // namespace UE::PixelStreaming2
 
-#endif // WITH_DEV_AUTOMATION_TESTS
+#endif // WITH_TESTS
diff --git a/Engine/Plugins/Media/PixelStreaming2/Source/PixelStreaming2RTC/Private/Tests/TestUtils.cpp b/Engine/Plugins/Media/PixelStreaming2/Source/PixelStreaming2RTC/Private/Tests/TestUtils.cpp
--- a/Engine/Plugins/Media/PixelStreaming2/Source/PixelStreaming2RTC/Private/Tests/TestUtils.cpp
+++ b/Engine/Plugins/Media/PixelStreaming2/Source/PixelStreaming2RTC/Private/Tests/TestUtils.cpp
@@ -14,7 +14,7 @@
 #include "SocketSubsystem.h"
 #include "SocketUtils.h"
 
-#if WITH_DEV_AUTOMATION_TESTS
+#if WITH_TESTS
 
 namespace UE::PixelStreaming2
 {
@@ -367,4 +367,4 @@ namespace UE::PixelStreaming2
 	}
 } // namespace UE::PixelStreaming2
 
-#endif // WITH_DEV_AUTOMATION_TESTS
+#endif // WITH_TESTS
diff --git a/Engine/Plugins/Media/PixelStreaming2/Source/PixelStreaming2RTC/Private/Tests/TestUtils.h b/Engine/Plugins/Media/PixelStreaming2/Source/PixelStreaming2RTC/Private/Tests/TestUtils.h
--- a/Engine/Plugins/Media/PixelStreaming2/Source/PixelStreaming2RTC/Private/Tests/TestUtils.h
+++ b/Engine/Plugins/Media/PixelStreaming2/Source/PixelStreaming2RTC/Private/Tests/TestUtils.h
@@ -9,7 +9,7 @@
 #include "VideoProducer.h"
 #include "Video/VideoConfig.h"
 
-#if WITH_DEV_AUTOMATION_TESTS
+#if WITH_TESTS
 
 namespace UE::PixelStreaming2
 {
@@ -73,4 +73,4 @@ namespace UE::PixelStreaming2
 	void											SetMediaDirection(EMediaType Media, EMediaDirection Direction);
 } // namespace UE::PixelStreaming2
 
-#endif // WITH_DEV_AUTOMATION_TESTS
\ No newline at end of file
+#endif // WITH_TESTS
\ No newline at end of file
