Fix for pixel streaming from android devices

#Synced-CL 0

diff --git a/Engine/Plugins/Media/PixelStreaming2/Source/PixelStreaming2RTC/Private/PixelStreaming2RTCModule.cpp b/Engine/Plugins/Media/PixelStreaming2/Source/PixelStreaming2RTC/Private/PixelStreaming2RTCModule.cpp
--- a/Engine/Plugins/Media/PixelStreaming2/Source/PixelStreaming2RTC/Private/PixelStreaming2RTCModule.cpp
+++ b/Engine/Plugins/Media/PixelStreaming2/Source/PixelStreaming2RTC/Private/PixelStreaming2RTCModule.cpp
@@ -52,6 +52,7 @@ namespace UE::PixelStreaming2
 	 */
 	extern "C" JNIEXPORT jlong JNICALL Java_org_webrtc_SoftwareVideoEncoderFactory_nativeCreateFactory(JNIEnv*, jclass);
 	extern "C" JNIEXPORT jlong JNICALL Java_org_webrtc_SoftwareVideoDecoderFactory_nativeCreateFactory(JNIEnv*, jclass);
+	extern "C" JNIEXPORT jlong JNICALL Java_org_webrtc_VideoEncoderFallback_nativeCreateEncoder(JNIEnv*, jclass, jobject, jobject);
 	extern "C" JNIEXPORT jlong JNICALL Java_org_webrtc_VideoDecoderFallback_nativeCreateDecoder(JNIEnv*, jclass, jobject, jobject);
 	extern "C" JNIEXPORT jlong JNICALL Java_org_webrtc_JniCommon_nativeAllocateByteBuffer(JNIEnv*, jclass, jint);
 	extern "C" JNIEXPORT jlong JNICALL Java_org_webrtc_YuvHelper_nativeCopyPlane(JNIEnv*, jclass, jobject, jint, jobject, jint, jint, jint);
@@ -63,6 +64,7 @@ namespace UE::PixelStreaming2
 	{
 		Java_org_webrtc_SoftwareVideoEncoderFactory_nativeCreateFactory(nullptr, nullptr);
 		Java_org_webrtc_SoftwareVideoDecoderFactory_nativeCreateFactory(nullptr, nullptr);
+		Java_org_webrtc_VideoEncoderFallback_nativeCreateEncoder(nullptr, nullptr, nullptr, nullptr);
 		Java_org_webrtc_VideoDecoderFallback_nativeCreateDecoder(nullptr, nullptr, nullptr, nullptr);
 		Java_org_webrtc_JniCommon_nativeAllocateByteBuffer(nullptr, nullptr, 0);
 		Java_org_webrtc_YuvHelper_nativeCopyPlane(nullptr, nullptr, nullptr, 0, nullptr, 0, 0, 0);
